UnrealCloudDDC - Report exceptions from our error controller to our tracer - to make sure they are tracked in datadog

diff --git a/Engine/Source/Programs/UnrealCloudDDC/Jupiter.Common/Controllers/ErrorController.cs b/Engine/Source/Programs/UnrealCloudDDC/Jupiter.Common/Controllers/ErrorController.cs
--- a/Engine/Source/Programs/UnrealCloudDDC/Jupiter.Common/Controllers/ErrorController.cs
+++ b/Engine/Source/Programs/UnrealCloudDDC/Jupiter.Common/Controllers/ErrorController.cs
@@ -1,6 +1,7 @@
 ﻿// Copyright Epic Games, Inc. All Rights Reserved.
 
 using System;
+using Datadog.Trace;
 using Microsoft.AspNetCore.Diagnostics;
 using Microsoft.AspNetCore.Mvc;
 
@@ -30,6 +31,9 @@ namespace Jupiter.Controllers
                 }
             }
 
+            // add the error to our trace
+            Tracer.Instance.ActiveScope.Span.SetException(context.Error);
+
             return Problem(
                 detail: context.Error.StackTrace,
                 title: context.Error.Message);
