UnrealCloudDDC - Tracing scylla scopes as a Scylla service in Datadog letting us more easily see how time of a request is spent doing db calls (this is similar to how the built in instrumentation would handle SQL or Mongo)

diff --git a/Engine/Source/Programs/UnrealCloudDDC/Jupiter/Implementation/Blob/ScyllaBlobIndex.cs b/Engine/Source/Programs/UnrealCloudDDC/Jupiter/Implementation/Blob/ScyllaBlobIndex.cs
--- a/Engine/Source/Programs/UnrealCloudDDC/Jupiter/Implementation/Blob/ScyllaBlobIndex.cs
+++ b/Engine/Source/Programs/UnrealCloudDDC/Jupiter/Implementation/Blob/ScyllaBlobIndex.cs
@@ -42,6 +42,7 @@ public class ScyllaBlobIndex : IBlobIndex
     {
         region ??= _jupiterSettings.CurrentValue.CurrentSite;
         using IScope scope = Tracer.Instance.StartActive("scylla.insert_blob_index");
+        scope.Span.ServiceName = IScyllaSessionManager.DatadogScyllaServiceName;
         scope.Span.ResourceName = $"{ns}.{id}";
 
         if (_scyllaSessionManager.IsScylla)
@@ -59,6 +60,7 @@ public class ScyllaBlobIndex : IBlobIndex
     public async Task<BlobInfo?> GetBlobInfo(NamespaceId ns, BlobIdentifier id, BlobIndexFlags flags = BlobIndexFlags.None)
     {
         using IScope scope = Tracer.Instance.StartActive("scylla.fetch_blob_index");
+        scope.Span.ServiceName = IScyllaSessionManager.DatadogScyllaServiceName;
         scope.Span.ResourceName = $"{ns}.{id}";
 
         bool includeReferences = (flags & BlobIndexFlags.IncludeReferences) != 0;
@@ -94,6 +96,7 @@ public class ScyllaBlobIndex : IBlobIndex
     public async Task<bool> RemoveBlobFromIndex(NamespaceId ns, BlobIdentifier id)
     {
         using IScope scope = Tracer.Instance.StartActive("scylla.remove_from_blob_index");
+        scope.Span.ServiceName = IScyllaSessionManager.DatadogScyllaServiceName;
         scope.Span.ResourceName = $"{ns}.{id}";
 
         if (_scyllaSessionManager.IsScylla)
@@ -114,6 +117,7 @@ public class ScyllaBlobIndex : IBlobIndex
     {
         region ??= _jupiterSettings.CurrentValue.CurrentSite;
         using IScope scope = Tracer.Instance.StartActive("scylla.remove_blob_index_region");
+        scope.Span.ServiceName = IScyllaSessionManager.DatadogScyllaServiceName;
         scope.Span.ResourceName = $"{ns}.{id}";
 
         if (_scyllaSessionManager.IsScylla)
@@ -140,7 +144,8 @@ public class ScyllaBlobIndex : IBlobIndex
 
     public async Task AddRefToBlobs(NamespaceId ns, BucketId bucket, IoHashKey key, BlobIdentifier[] blobs)
     {
-        using IScope _ = Tracer.Instance.StartActive("scylla.add_ref_blobs");
+        using IScope scope = Tracer.Instance.StartActive("scylla.add_ref_blobs");
+        scope.Span.ServiceName = IScyllaSessionManager.DatadogScyllaServiceName;
 
         string nsAsString = ns.ToString();
         ScyllaObjectReference reference = new ScyllaObjectReference(bucket, key);
@@ -174,7 +179,8 @@ public class ScyllaBlobIndex : IBlobIndex
 
     public async IAsyncEnumerable<BlobInfo> GetAllBlobs()
     {
-        using IScope _ = Tracer.Instance.StartActive("scylla.get_all_blobs");
+        using IScope scope = Tracer.Instance.StartActive("scylla.get_all_blobs");
+        scope.Span.ServiceName = IScyllaSessionManager.DatadogScyllaServiceName;
 
         string cqlOptions = _scyllaSessionManager.IsScylla ? "BYPASS CACHE" : "";
 
diff --git a/Engine/Source/Programs/UnrealCloudDDC/Jupiter/Implementation/ContentId/ScyllaContentIdStore.cs b/Engine/Source/Programs/UnrealCloudDDC/Jupiter/Implementation/ContentId/ScyllaContentIdStore.cs
--- a/Engine/Source/Programs/UnrealCloudDDC/Jupiter/Implementation/ContentId/ScyllaContentIdStore.cs
+++ b/Engine/Source/Programs/UnrealCloudDDC/Jupiter/Implementation/ContentId/ScyllaContentIdStore.cs
@@ -51,7 +51,8 @@ namespace Jupiter.Implementation
 
             {
                 using IScope contentIdFetchScope = Tracer.Instance.StartActive("ScyllaContentIdStore.FetchContentId");
-                scope.Span.ResourceName = contentId.ToString();
+                contentIdFetchScope.Span.ResourceName = contentId.ToString();
+                contentIdFetchScope.Span.ServiceName = IScyllaSessionManager.DatadogScyllaServiceName;
 
                 if (_scyllaSessionManager.IsScylla)
                 {
@@ -121,6 +122,9 @@ namespace Jupiter.Implementation
 
         public async Task Put(NamespaceId ns, ContentId contentId, BlobIdentifier blobIdentifier, int contentWeight)
         {
+            using IScope scope = Tracer.Instance.StartActive("ScyllaContentIdStore.PutContentId");
+            scope.Span.ServiceName = IScyllaSessionManager.DatadogScyllaServiceName;
+
             if (_scyllaSessionManager.IsScylla)
             {
                 await _mapper.UpdateAsync<ScyllaContentId>("SET chunks = ? WHERE content_id = ? AND content_weight = ?", new [] {new ScyllaBlobIdentifier(blobIdentifier)}, new ScyllaBlobIdentifier(contentId), contentWeight);
diff --git a/Engine/Source/Programs/UnrealCloudDDC/Jupiter/Implementation/Objects/ScyllaReferencesStore.cs b/Engine/Source/Programs/UnrealCloudDDC/Jupiter/Implementation/Objects/ScyllaReferencesStore.cs
--- a/Engine/Source/Programs/UnrealCloudDDC/Jupiter/Implementation/Objects/ScyllaReferencesStore.cs
+++ b/Engine/Source/Programs/UnrealCloudDDC/Jupiter/Implementation/Objects/ScyllaReferencesStore.cs
@@ -8,7 +8,6 @@ using Cassandra;
 using Cassandra.Mapping;
 using Datadog.Trace;
 using EpicGames.Horde.Storage;
-using Jupiter.Implementation;
 using Microsoft.Extensions.Options;
 using Serilog;
 
@@ -58,6 +57,7 @@ namespace Jupiter.Implementation
         public async Task<ObjectRecord> Get(NamespaceId ns, BucketId bucket, IoHashKey name, IReferencesStore.FieldFlags flags)
         {
             using IScope scope = Tracer.Instance.StartActive("scylla.get");
+            scope.Span.ServiceName = IScyllaSessionManager.DatadogScyllaServiceName;
             scope.Span.ResourceName = $"{ns}.{bucket}.{name}";
 
             ScyllaObject? o;
@@ -93,6 +93,7 @@ namespace Jupiter.Implementation
         public async Task Put(NamespaceId ns, BucketId bucket, IoHashKey name, BlobIdentifier blobHash, byte[] blob, bool isFinalized)
         {
             using IScope scope = Tracer.Instance.StartActive("scylla.put");
+            scope.Span.ServiceName = IScyllaSessionManager.DatadogScyllaServiceName;
             scope.Span.ResourceName = $"{ns}.{bucket}.{name}";
 
             if (blob.LongLength > _settings.CurrentValue.InlineBlobMaxSize)
@@ -111,6 +112,7 @@ namespace Jupiter.Implementation
         public async Task Finalize(NamespaceId ns, BucketId bucket, IoHashKey name, BlobIdentifier blobIdentifier)
         {
             using IScope scope = Tracer.Instance.StartActive("scylla.finalize");
+            scope.Span.ServiceName = IScyllaSessionManager.DatadogScyllaServiceName;
             scope.Span.ResourceName = $"{ns}.{bucket}.{name}";
 
             await _mapper.UpdateAsync<ScyllaObject>("SET is_finalized=true WHERE namespace=? AND bucket=? AND name=?", ns.ToString(), bucket.ToString(), name.ToString());
@@ -118,15 +120,16 @@ namespace Jupiter.Implementation
 
         public async Task UpdateLastAccessTime(NamespaceId ns, BucketId bucket, IoHashKey name, DateTime lastAccessTime)
         {
-            using IScope _ = Tracer.Instance.StartActive("scylla.update_last_access_time");
+            using IScope scope = Tracer.Instance.StartActive("scylla.update_last_access_time");
+            scope.Span.ServiceName = IScyllaSessionManager.DatadogScyllaServiceName;
 
             await _mapper.UpdateAsync<ScyllaObject>("SET last_access_time = ? WHERE namespace = ? AND bucket = ? AND name = ?", lastAccessTime, ns.ToString(), bucket.ToString(), name.ToString());
         }
 
         public async IAsyncEnumerable<(BucketId, IoHashKey, DateTime)> GetRecords(NamespaceId ns)
         {
-            using IScope _ = Tracer.Instance.StartActive("scylla.get_records");
-
+            using IScope scope = Tracer.Instance.StartActive("scylla.get_records");
+            scope.Span.ServiceName = IScyllaSessionManager.DatadogScyllaServiceName;
             int retryAttempts = 0;
             const int MaxRetryAttempts = 3;
             RowSet rowSet = await _session.ExecuteAsync(_getObjectsStatement.Bind(ns.ToString()));
@@ -172,7 +175,8 @@ namespace Jupiter.Implementation
 
         public async IAsyncEnumerable<NamespaceId> GetNamespaces()
         {
-            using IScope _ = Tracer.Instance.StartActive("scylla.get_namespaces");
+            using IScope scope = Tracer.Instance.StartActive("scylla.get_namespaces");
+            scope.Span.ServiceName = IScyllaSessionManager.DatadogScyllaServiceName;
             RowSet rowSet = await _session.ExecuteAsync(_getNamespacesStatement.Bind());
 
             foreach (Row row in rowSet)
@@ -189,6 +193,7 @@ namespace Jupiter.Implementation
         public async Task<bool> Delete(NamespaceId ns, BucketId bucket, IoHashKey key)
         {
             using IScope scope = Tracer.Instance.StartActive("scylla.delete_record");
+            scope.Span.ServiceName = IScyllaSessionManager.DatadogScyllaServiceName;
             scope.Span.ResourceName = $"{ns}.{bucket}.{key}";
 
             AppliedInfo<ScyllaObject> info = await _mapper.DeleteIfAsync<ScyllaObject>("WHERE namespace=? AND bucket=? AND name=? IF EXISTS", ns.ToString(), bucket.ToString(), key.ToString());
@@ -203,7 +208,8 @@ namespace Jupiter.Implementation
 
         public async Task<long> DropNamespace(NamespaceId ns)
         {
-            using IScope _ = Tracer.Instance.StartActive("scylla.delete_namespace");
+            using IScope scope = Tracer.Instance.StartActive("scylla.delete_namespace");
+            scope.Span.ServiceName = IScyllaSessionManager.DatadogScyllaServiceName;
             RowSet rowSet = await _session.ExecuteAsync(new SimpleStatement("SELECT bucket, name FROM objects WHERE namespace = ? ALLOW FILTERING;", ns.ToString()));
             long deletedCount = 0;
             foreach (Row row in rowSet)
@@ -224,7 +230,8 @@ namespace Jupiter.Implementation
 
         public async Task<long> DeleteBucket(NamespaceId ns, BucketId bucket)
         {
-            using IScope _ = Tracer.Instance.StartActive("scylla.delete_bucket");
+            using IScope scope = Tracer.Instance.StartActive("scylla.delete_bucket");
+            scope.Span.ServiceName = IScyllaSessionManager.DatadogScyllaServiceName;
             RowSet rowSet = await _session.ExecuteAsync(new SimpleStatement("SELECT name FROM objects WHERE namespace = ? AND bucket = ? ALLOW FILTERING;", ns.ToString(), bucket.ToString()));
             long deletedCount = 0;
             foreach (Row row in rowSet)
@@ -243,7 +250,8 @@ namespace Jupiter.Implementation
 
         private async Task MaybeAddBucket(NamespaceId ns, BucketId bucket)
         {
-            using IScope _ = Tracer.Instance.StartActive("scylla.add_bucket");
+            using IScope scope = Tracer.Instance.StartActive("scylla.add_bucket");
+            scope.Span.ServiceName = IScyllaSessionManager.DatadogScyllaServiceName;
             await _mapper.UpdateAsync<ScyllaBucket>("SET bucket = bucket + ? WHERE namespace = ?", new string[] {bucket.ToString()}, ns.ToString());
         }
     }
diff --git a/Engine/Source/Programs/UnrealCloudDDC/Jupiter/Implementation/ScyllaSessionManager.cs b/Engine/Source/Programs/UnrealCloudDDC/Jupiter/Implementation/ScyllaSessionManager.cs
--- a/Engine/Source/Programs/UnrealCloudDDC/Jupiter/Implementation/ScyllaSessionManager.cs
+++ b/Engine/Source/Programs/UnrealCloudDDC/Jupiter/Implementation/ScyllaSessionManager.cs
@@ -18,6 +18,8 @@ namespace Jupiter.Implementation
         /// True when using a base cassandra protocol
         /// </summary>
         bool IsCassandra { get; }
+
+        public const string DatadogScyllaServiceName = "ScyllaDB";
     }
 
     public class ScyllaSessionManager : IScyllaSessionManager
diff --git a/Engine/Source/Programs/UnrealCloudDDC/Jupiter/Implementation/TransactionLog/ScyllaReplicationLog.cs b/Engine/Source/Programs/UnrealCloudDDC/Jupiter/Implementation/TransactionLog/ScyllaReplicationLog.cs
--- a/Engine/Source/Programs/UnrealCloudDDC/Jupiter/Implementation/TransactionLog/ScyllaReplicationLog.cs
+++ b/Engine/Source/Programs/UnrealCloudDDC/Jupiter/Implementation/TransactionLog/ScyllaReplicationLog.cs
@@ -75,7 +75,8 @@ namespace Jupiter.Implementation
 
         public async Task<(string, Guid)> InsertAddEvent(NamespaceId ns, BucketId bucket, IoHashKey key, BlobIdentifier objectBlob, DateTime? timestamp)
         {
-            using IScope _ = Tracer.Instance.StartActive("scylla.insert_add_event");
+            using IScope scope = Tracer.Instance.StartActive("scylla.insert_add_event");
+            scope.Span.ServiceName = IScyllaSessionManager.DatadogScyllaServiceName;
             Task addNamespaceTask = PotentiallyAddNamespace(ns);
             DateTime timeBucket = timestamp.GetValueOrDefault(DateTime.UtcNow);
             ScyllaReplicationLogEvent log = new ScyllaReplicationLogEvent(ns.ToString(), bucket.ToString(), key.ToString(), timeBucket, ScyllaReplicationLogEvent.OpType.Added, objectBlob);
@@ -87,7 +88,8 @@ namespace Jupiter.Implementation
 
         public async Task<(string, Guid)> InsertDeleteEvent(NamespaceId ns, BucketId bucket, IoHashKey key, DateTime? timestamp)
         {
-            using IScope _ = Tracer.Instance.StartActive("scylla.insert_delete_event");
+            using IScope scope = Tracer.Instance.StartActive("scylla.insert_delete_event");
+            scope.Span.ServiceName = IScyllaSessionManager.DatadogScyllaServiceName;
             Task addNamespaceTask = PotentiallyAddNamespace(ns);
             DateTime timeBucket = timestamp.GetValueOrDefault(DateTime.UtcNow);
             ScyllaReplicationLogEvent log = new ScyllaReplicationLogEvent(ns.ToString(), bucket.ToString(), key.ToString(), timeBucket, ScyllaReplicationLogEvent.OpType.Deleted, null);
@@ -105,7 +107,7 @@ namespace Jupiter.Implementation
         public async IAsyncEnumerable<ReplicationLogEvent> Get(NamespaceId ns, string? lastBucket, Guid? lastEvent)
         {
             using IScope getReplicationLogScope = Tracer.Instance.StartActive("scylla.get_replication_log");
-
+            getReplicationLogScope.Span.ServiceName = IScyllaSessionManager.DatadogScyllaServiceName;
             if (lastBucket == "now")
             {
                 // for debug purposes we allow you to list the latest bucket
@@ -119,6 +121,7 @@ namespace Jupiter.Implementation
             await foreach (long bucketField in buckets)
             {
                 using IScope readReplicationScope = Tracer.Instance.StartActive("scylla.read_replication_bucket");
+                readReplicationScope.Span.ServiceName = IScyllaSessionManager.DatadogScyllaServiceName;
                 DateTime t = DateTime.FromFileTimeUtc(bucketField);
                 string bucket = t.ToReplicationBucketIdentifier();
                 readReplicationScope.Span.ResourceName = bucket;
@@ -180,7 +183,7 @@ namespace Jupiter.Implementation
         private async IAsyncEnumerable<long> FindReplicationBuckets(NamespaceId ns, string? lastBucket)
         {
             using IScope findReplicationBucketScope = Tracer.Instance.StartActive("scylla.find_replication_buckets");
-
+            findReplicationBucketScope.Span.ServiceName = IScyllaSessionManager.DatadogScyllaServiceName;
             // ignore any bucket that is older then a cutoff, as that can cause us to end up scanning thru a lot of hours that will never exist (incremental logs are deleted after 7 days)
             DateTime oldCutoff = DateTime.UtcNow.AddDays(-14);
 
@@ -200,8 +203,8 @@ namespace Jupiter.Implementation
             }
             else
             {
-                using IScope _ = Tracer.Instance.StartActive("scylla.determine_first_replication_bucket");
-
+                using IScope firstReplicationBucketScope = Tracer.Instance.StartActive("scylla.determine_first_replication_bucket");
+                firstReplicationBucketScope.Span.ServiceName = IScyllaSessionManager.DatadogScyllaServiceName;
                 // we should have no data older then the ttl to lets just assume that the bucket to start searching from is now - time to live
                 DateTime oldestTimestamp = DateTime.UtcNow.AddSeconds(-1 * _settings.CurrentValue.ReplicationLogTimeToLive.TotalSeconds);
                 startBucketTime = oldestTimestamp;
@@ -211,7 +214,8 @@ namespace Jupiter.Implementation
             DateTime bucketTime = startBucketTime.AddHours(1.0).ToHourlyBucket();
             while(bucketTime < DateTime.UtcNow && bucketTime > oldCutoff)
             {
-                using IScope _ = Tracer.Instance.StartActive("scylla.determine_replication_bucket_exists");
+                using IScope determineBucketExistsScope = Tracer.Instance.StartActive("scylla.determine_replication_bucket_exists");
+                determineBucketExistsScope.Span.ServiceName = IScyllaSessionManager.DatadogScyllaServiceName;
                 // fetch all the buckets that exists and sort them based on time
                 IEnumerable<ScyllaReplicationLogEvent> logEvent = await _mapper.FetchAsync<ScyllaReplicationLogEvent>("WHERE namespace = ? AND replication_bucket = ? LIMIT 1", ns.ToString(), bucketTime.ToFileTimeUtc());
                 ScyllaReplicationLogEvent? e = logEvent.FirstOrDefault();
@@ -286,6 +290,9 @@ namespace Jupiter.Implementation
 
         public async Task<ReplicatorState?> GetReplicatorState(NamespaceId ns, string name)
         {
+            using IScope scope = Tracer.Instance.StartActive("scylla.get_replicator_state");
+            scope.Span.ServiceName = IScyllaSessionManager.DatadogScyllaServiceName;
+
             ScyllaReplicationState? replicationState = await _mapper.FirstOrDefaultAsync<ScyllaReplicationState>("WHERE namespace = ? AND name = ?", ns.ToString(), name);
 
             return replicationState?.ToReplicatorState();
