UnrealCloudDDC - Prefetch pages during ref GC, should make the GC a bit faster but more importantly it lets us handle any cassandra read timeout exceptions and attempt retries

diff --git a/Engine/Source/Programs/UnrealCloudDDC/Jupiter/Implementation/Objects/ScyllaReferencesStore.cs b/Engine/Source/Programs/UnrealCloudDDC/Jupiter/Implementation/Objects/ScyllaReferencesStore.cs
--- a/Engine/Source/Programs/UnrealCloudDDC/Jupiter/Implementation/Objects/ScyllaReferencesStore.cs
+++ b/Engine/Source/Programs/UnrealCloudDDC/Jupiter/Implementation/Objects/ScyllaReferencesStore.cs
@@ -133,13 +133,23 @@ namespace Jupiter.Implementation
             int retryAttempts = 0;
             const int MaxRetryAttempts = 3;
             RowSet rowSet = await _session.ExecuteAsync(_getObjectsStatement.Bind(ns.ToString()));
+
+            Task? prefetchTask = null;
             foreach (Row row in rowSet)
             {
                 try
                 {
-                    if (rowSet.GetAvailableWithoutFetching() == 0)
+                    // when we only have a fow rows left we start prefetching the next page
+                    // this also allows us to get any read timeouts here instead of in the foreach line that moves the enumerator thus enabling us to handle it
+                    if (prefetchTask == null && rowSet.GetAvailableWithoutFetching() < 5)
+                    {
+                        prefetchTask = rowSet.FetchMoreResultsAsync();
+                    }
+
+                    if (prefetchTask is { IsCompleted: true })
                     {
-                        await rowSet.FetchMoreResultsAsync();
+                        await prefetchTask;
+                        prefetchTask = null;
                     }
                 }
                 catch (ReadTimeoutException)
diff --git a/Engine/Source/Programs/UnrealCloudDDC/changelog.md b/Engine/Source/Programs/UnrealCloudDDC/changelog.md
--- a/Engine/Source/Programs/UnrealCloudDDC/changelog.md
+++ b/Engine/Source/Programs/UnrealCloudDDC/changelog.md
@@ -1,4 +1,5 @@
 # Unreleased
+* Improved handling of cassandra timeouts during ref GC.
 
 # 0.1.0
 * First public release of UnrealCloudDDC
