UnrealCloudDDC - Set editorconfig to require space instead of tabs so we follow dotnet standard conventions, this is already how to code looks.

diff --git a/Engine/Source/Programs/UnrealCloudDDC/.editorconfig b/Engine/Source/Programs/UnrealCloudDDC/.editorconfig
--- a/Engine/Source/Programs/UnrealCloudDDC/.editorconfig
+++ b/Engine/Source/Programs/UnrealCloudDDC/.editorconfig
@@ -1108,4 +1108,5 @@ dotnet_diagnostic.ca5351.severity = none
 
 # CA1716: Identifiers should not match keywords
 # Reason: We frequently use verbs as method names, like GET as well as terms like namespace which collide with built in keywords
-dotnet_diagnostic.CA1716.severity = none
\ No newline at end of file
+dotnet_diagnostic.CA1716.severity = none
+indent_style=space
\ No newline at end of file
