UnrealCloudDDC - Removing references that are no longer valid during blob GC

diff --git a/Engine/Source/Programs/UnrealCloudDDC/Jupiter/Implementation/Blob/CachedBlobIndex.cs b/Engine/Source/Programs/UnrealCloudDDC/Jupiter/Implementation/Blob/CachedBlobIndex.cs
--- a/Engine/Source/Programs/UnrealCloudDDC/Jupiter/Implementation/Blob/CachedBlobIndex.cs
+++ b/Engine/Source/Programs/UnrealCloudDDC/Jupiter/Implementation/Blob/CachedBlobIndex.cs
@@ -4,8 +4,6 @@ using System;
 using System.Collections.Generic;
 using System.Threading.Tasks;
 using EpicGames.Horde.Storage;
-using Jupiter;
-using Jupiter.Implementation;
 using Microsoft.Extensions.Options;
 
 namespace Jupiter.Implementation.Blob;
@@ -71,4 +69,10 @@ public class CachedBlobIndex : IBlobIndex
     {
         throw new NotImplementedException();
     }
+
+    public Task RemoveReferences(NamespaceId ns, BlobIdentifier id, List<(BucketId, IoHashKey)> references)
+    {
+        // We do not actually track any blob information when running in cached mode
+        return Task.CompletedTask;
+    }
 }
diff --git a/Engine/Source/Programs/UnrealCloudDDC/Jupiter/Implementation/Blob/IBlobIndex.cs b/Engine/Source/Programs/UnrealCloudDDC/Jupiter/Implementation/Blob/IBlobIndex.cs
--- a/Engine/Source/Programs/UnrealCloudDDC/Jupiter/Implementation/Blob/IBlobIndex.cs
+++ b/Engine/Source/Programs/UnrealCloudDDC/Jupiter/Implementation/Blob/IBlobIndex.cs
@@ -3,7 +3,6 @@
 using System.Collections.Generic;
 using System.Threading.Tasks;
 using EpicGames.Horde.Storage;
-using Jupiter.Implementation;
 
 namespace Jupiter.Implementation.Blob
 {
@@ -32,5 +31,6 @@ namespace Jupiter.Implementation.Blob
         Task<bool> BlobExistsInRegion(NamespaceId ns, BlobIdentifier blobIdentifier);
         Task AddRefToBlobs(NamespaceId ns, BucketId bucket, IoHashKey key, BlobIdentifier[] blobs);
         IAsyncEnumerable<BlobInfo> GetAllBlobs();
+        Task RemoveReferences(NamespaceId ns, BlobIdentifier id, List<(BucketId,IoHashKey)> references);
     }
 }
diff --git a/Engine/Source/Programs/UnrealCloudDDC/Jupiter/Implementation/Blob/MemoryBlobIndex.cs b/Engine/Source/Programs/UnrealCloudDDC/Jupiter/Implementation/Blob/MemoryBlobIndex.cs
--- a/Engine/Source/Programs/UnrealCloudDDC/Jupiter/Implementation/Blob/MemoryBlobIndex.cs
+++ b/Engine/Source/Programs/UnrealCloudDDC/Jupiter/Implementation/Blob/MemoryBlobIndex.cs
@@ -4,8 +4,6 @@ using System.Collections.Concurrent;
 using System.Collections.Generic;
 using System.Threading.Tasks;
 using EpicGames.Horde.Storage;
-using Jupiter;
-using Jupiter.Implementation;
 using Microsoft.Extensions.Options;
 
 namespace Jupiter.Implementation.Blob;
@@ -113,6 +111,21 @@ public class MemoryBlobIndex : IBlobIndex
         }
     }
 
+    public Task RemoveReferences(NamespaceId ns, BlobIdentifier id, List<(BucketId,IoHashKey)> references)
+    {
+        ConcurrentDictionary<BlobIdentifier, MemoryBlobInfo> index = GetNamespaceContainer(ns);
+
+        if (index.TryGetValue(id, out MemoryBlobInfo? blobInfo))
+        {
+            foreach ((BucketId bucketId, IoHashKey key) in references)
+            {
+                blobInfo.References?.Remove((bucketId, key));   
+            }
+        }
+
+        return Task.CompletedTask;
+    }
+
     private static MemoryBlobInfo NewBlobInfo(NamespaceId ns, BlobIdentifier blob, string region)
     {
         MemoryBlobInfo info = new MemoryBlobInfo
diff --git a/Engine/Source/Programs/UnrealCloudDDC/Jupiter/Implementation/Blob/MongoBlobIndex.cs b/Engine/Source/Programs/UnrealCloudDDC/Jupiter/Implementation/Blob/MongoBlobIndex.cs
--- a/Engine/Source/Programs/UnrealCloudDDC/Jupiter/Implementation/Blob/MongoBlobIndex.cs
+++ b/Engine/Source/Programs/UnrealCloudDDC/Jupiter/Implementation/Blob/MongoBlobIndex.cs
@@ -4,8 +4,6 @@ using System.Collections.Generic;
 using System.Linq;
 using System.Threading.Tasks;
 using EpicGames.Horde.Storage;
-using Jupiter;
-using Jupiter.Implementation;
 using Microsoft.Extensions.Options;
 using MongoDB.Bson.Serialization.Attributes;
 using MongoDB.Bson.Serialization.Options;
@@ -131,6 +129,18 @@ public class MongoBlobIndex : MongoStore, IBlobIndex
             }
         }
     }
+
+    public async Task RemoveReferences(NamespaceId ns, BlobIdentifier id, List<(BucketId,IoHashKey)> references)
+    {
+        IMongoCollection<MongoBlobIndexModelV0> collection = GetCollection<MongoBlobIndexModelV0>();
+
+        string nsAsString = ns.ToString();
+        List<Dictionary<string, string>> refs = references.Select(tuple => new Dictionary<string, string>() { {tuple.Item1.ToString(), tuple.Item2.ToString()}}).ToList();
+        UpdateDefinition<MongoBlobIndexModelV0> update = Builders<MongoBlobIndexModelV0>.Update.PullAll(m => m.References, refs);
+        FilterDefinition<MongoBlobIndexModelV0> filter = Builders<MongoBlobIndexModelV0>.Filter.Where(m => m.Ns == nsAsString && m.BlobId == id.ToString());
+
+        await collection.FindOneAndUpdateAsync(filter, update);
+    }
 }
 
 [BsonDiscriminator("blob-index.v0")]
diff --git a/Engine/Source/Programs/UnrealCloudDDC/Jupiter/Implementation/Blob/ScyllaBlobIndex.cs b/Engine/Source/Programs/UnrealCloudDDC/Jupiter/Implementation/Blob/ScyllaBlobIndex.cs
--- a/Engine/Source/Programs/UnrealCloudDDC/Jupiter/Implementation/Blob/ScyllaBlobIndex.cs
+++ b/Engine/Source/Programs/UnrealCloudDDC/Jupiter/Implementation/Blob/ScyllaBlobIndex.cs
@@ -197,6 +197,42 @@ public class ScyllaBlobIndex : IBlobIndex
             }
         }
     }
+
+    public async Task RemoveReferences(NamespaceId ns, BlobIdentifier id, List<(BucketId,IoHashKey)> references)
+    {
+        using IScope scope = Tracer.Instance.StartActive("scylla.remove_ref_blobs");
+        ScyllaUtils.SetupScyllaScope(scope);
+
+        string nsAsString = ns.ToString();
+        (BucketId, IoHashKey)[] blobs = references.ToArray();
+        Task[] refUpdateTasks = new Task[blobs.Length];
+        for (int i = 0; i < blobs.Length; i++)
+        {
+            (BucketId bucket, IoHashKey key) = blobs[i];
+            ScyllaObjectReference reference = new ScyllaObjectReference(bucket, key);
+            Task updateTask;
+            if (_scyllaSessionManager.IsScylla)
+            {
+                updateTask = _mapper.UpdateAsync<ScyllaBlobIndexTable>(
+                    "SET references = references - ? WHERE namespace = ? AND blob_id = ?",
+                    new[] { reference },
+                    nsAsString,
+                    new ScyllaBlobIdentifier(id));
+            }
+            else
+            {
+                updateTask = _mapper.UpdateAsync<CassandraBlobIndexTable>(
+                    "SET references = references - ? WHERE namespace = ? AND blob_id = ?",
+                    new[] { reference },
+                    nsAsString,
+                    id.HashData);
+            }
+
+            refUpdateTasks[i] = updateTask;
+        }
+
+        await Task.WhenAll(refUpdateTasks);
+    }
 }
 
 [Cassandra.Mapping.Attributes.Table("blob_index")]
diff --git a/Engine/Source/Programs/UnrealCloudDDC/Jupiter/Implementation/GC/OrphanBlobCleanupRefs.cs b/Engine/Source/Programs/UnrealCloudDDC/Jupiter/Implementation/GC/OrphanBlobCleanupRefs.cs
--- a/Engine/Source/Programs/UnrealCloudDDC/Jupiter/Implementation/GC/OrphanBlobCleanupRefs.cs
+++ b/Engine/Source/Programs/UnrealCloudDDC/Jupiter/Implementation/GC/OrphanBlobCleanupRefs.cs
@@ -9,9 +9,7 @@ using Dasync.Collections;
 using Datadog.Trace;
 using EpicGames.Horde.Storage;
 using Jupiter.Implementation.Blob;
-using Jupiter;
 using Jupiter.Common;
-using Jupiter.Implementation;
 using Microsoft.Extensions.Options;
 using Serilog;
 
@@ -140,6 +138,8 @@ namespace Jupiter.Implementation
                     continue;
                 }
 
+                List<(BucketId, IoHashKey)> oldReferences = new List<(BucketId, IoHashKey)>();
+
                 foreach ((BucketId, IoHashKey) tuple in blobIndex.References!)
                 {
                     try
@@ -152,8 +152,15 @@ namespace Jupiter.Implementation
                     catch (ObjectNotFoundException)
                     {
                         // this is not a valid reference so we should delete
+                        oldReferences.Add(tuple);
                     }
                 }
+
+                if (found)
+                {
+                    // if the object is still alive but had old references we remove the old references to keep the size of the references array more reasonable
+                    await _blobIndex.RemoveReferences(blobNamespace, blob, oldReferences);
+                }
             }
 
             if (cancellationToken.IsCancellationRequested)
