UnrealCloudDDC - Log warnings when scylla reads timeout to help us debug read timeouts when scanning all references

diff --git a/Engine/Source/Programs/UnrealCloudDDC/Jupiter/Implementation/Objects/ScyllaReferencesStore.cs b/Engine/Source/Programs/UnrealCloudDDC/Jupiter/Implementation/Objects/ScyllaReferencesStore.cs
--- a/Engine/Source/Programs/UnrealCloudDDC/Jupiter/Implementation/Objects/ScyllaReferencesStore.cs
+++ b/Engine/Source/Programs/UnrealCloudDDC/Jupiter/Implementation/Objects/ScyllaReferencesStore.cs
@@ -156,12 +156,15 @@ namespace Jupiter.Implementation
                 {
                     if (retryAttempts < MaxRetryAttempts)
                     {
+                        retryAttempts += 1;
+                        _logger.Warning("Cassandra read timeouts, waiting a while and then retrying. Attempt {Attempts} .", retryAttempts);
                         // wait 10 seconds and try again as the Db is under heavy load right now
                         await Task.Delay(TimeSpan.FromSeconds(10));
-                        retryAttempts += 1;
 
                         continue;
                     }
+
+                    _logger.Warning("Cassandra read timeouts, attempted {Attempts} attempts now we give up.", retryAttempts);
                     // we have failed to many times, rethrow the exception and abort to avoid stalling here for ever
                     throw;
                 }
