UnrealCloudDDC - Tweaked the OpenTelemetry tracing, upgrading to newer libraries that gives us better options to enrich and adding the operation name to scylla.

Removed Azure telemetry exporter as it was a hassle to configure in a optional way, its better if people just use the OTLP path anyway

diff --git a/Engine/Source/Programs/UnrealCloudDDC/Jupiter.Common/BaseStartup.cs b/Engine/Source/Programs/UnrealCloudDDC/Jupiter.Common/BaseStartup.cs
--- a/Engine/Source/Programs/UnrealCloudDDC/Jupiter.Common/BaseStartup.cs
+++ b/Engine/Source/Programs/UnrealCloudDDC/Jupiter.Common/BaseStartup.cs
@@ -9,7 +9,6 @@ using System.Net.Mime;
 using System.Reflection;
 using System.Text;
 using Amazon;
-using Azure.Monitor.OpenTelemetry.Exporter;
 using EpicGames.AspNet;
 using Jupiter.Common;
 using Microsoft.AspNetCore.Authentication;
@@ -250,33 +249,18 @@ namespace Jupiter
                 builder.AddHttpClientInstrumentation();
                 builder.AddAspNetCoreInstrumentation();
 
-                builder.Configure((provider, providerBuilder) =>
+                builder.AddOtlpExporter();
+
+                builder.ConfigureBuilder((provider, _) =>
                 {
                     IOptionsMonitor<TracerSettings> settings = provider.GetService<IOptionsMonitor<TracerSettings>>()!;
                     string tracerServiceName = settings.CurrentValue.TracerServiceName ?? "UnrealCloudDDC";
 
                     builder.AddSource(tracerServiceName, "ScyllaDB");
 
-                    builder.SetResourceBuilder(ResourceBuilder.CreateDefault().AddService(tracerServiceName, serviceVersion: settings.CurrentValue.TracerServiceVersion, serviceInstanceId: Environment.MachineName));
-
-                    if (settings.CurrentValue.TracerImplementation == TracerSettings.TracerImplementations.OpenTelemetry)
-                    {
-                        providerBuilder.AddOtlpExporter(options =>
-                        {
-                            if (settings.CurrentValue.ConnectionString != null)
-                            {
-                                options.Endpoint = new Uri(settings.CurrentValue.ConnectionString);
-                            }
-                        });
-                    }
-                    else if (settings.CurrentValue.TracerImplementation == TracerSettings.TracerImplementations.AzureMonitor)
-                    {
-                        // forward open telemetry requests to AzureMonitor
-                        providerBuilder.AddAzureMonitorTraceExporter(options =>
-                        {
-                            options.ConnectionString = settings.CurrentValue.ConnectionString;
-                        });
-                    }
+                    builder.SetResourceBuilder(ResourceBuilder.CreateDefault().AddService(tracerServiceName,
+                        serviceVersion: settings.CurrentValue.TracerServiceVersion,
+                        serviceInstanceId: Environment.MachineName));
                 });
             });
             services.Configure<OpenTelemetryLoggerOptions>(opt =>
@@ -685,20 +669,6 @@ namespace Jupiter
 
     public class TracerSettings
     {
-        public enum TracerImplementations 
-        { 
-            /// <summary>
-            /// Use OpenTelemetry with OTLP exporter
-            /// </summary>
-            OpenTelemetry, 
-            /// <summary>
-            /// Forward all open telemetry events to azure monitor
-            /// </summary>
-            AzureMonitor
-        };
-
-        public TracerImplementations TracerImplementation { get; set; } = TracerImplementations.OpenTelemetry;
-        public string? ConnectionString { get; set; } = null;
         public string? TracerServiceName { get; set; } = null;
         public string? TracerServiceVersion { get; set; } = null;
     }
diff --git a/Engine/Source/Programs/UnrealCloudDDC/Jupiter.Common/Jupiter.Common.csproj b/Engine/Source/Programs/UnrealCloudDDC/Jupiter.Common/Jupiter.Common.csproj
--- a/Engine/Source/Programs/UnrealCloudDDC/Jupiter.Common/Jupiter.Common.csproj
+++ b/Engine/Source/Programs/UnrealCloudDDC/Jupiter.Common/Jupiter.Common.csproj
@@ -19,12 +19,11 @@
     <PackageReference Include="AWSSDK.SecurityToken" Version="3.7.1.4" />
     <PackageReference Include="Azure.Identity" Version="1.5.0" />
     <PackageReference Include="Azure.Security.KeyVault.Secrets" Version="4.2.0" />
-    <PackageReference Include="OpenTelemetry" Version="1.3.1" />
-    <PackageReference Include="OpenTelemetry.Exporter.OpenTelemetryProtocol" Version="1.3.1" />
-    <PackageReference Include="OpenTelemetry.Instrumentation.AspNetCore" Version="1.0.0-rc9.4" />
-    <PackageReference Include="OpenTelemetry.Instrumentation.Http" Version="1.0.0-rc9.4" />
-    <PackageReference Include="OpenTelemetry.Extensions.Hosting" Version="1.0.0-rc9.4" />
-    <PackageReference Include="Azure.Monitor.OpenTelemetry.Exporter" Version="1.0.0-beta.4" />
+    <PackageReference Include="OpenTelemetry" Version="1.4.0-beta.3" />
+    <PackageReference Include="OpenTelemetry.Exporter.OpenTelemetryProtocol" Version="1.4.0-beta.3" />
+    <PackageReference Include="OpenTelemetry.Instrumentation.AspNetCore" Version="1.0.0-rc9.9" />
+    <PackageReference Include="OpenTelemetry.Instrumentation.Http" Version="1.0.0-rc9.9" />
+    <PackageReference Include="OpenTelemetry.Extensions.Hosting" Version="1.0.0-rc9.9" />
     <PackageReference Include="IdentityServer4.AccessTokenValidation" Version="3.0.1" />
 
     <PackageReference Include="Microsoft.AspNet.WebApi.Client" Version="5.2.7" />
diff --git a/Engine/Source/Programs/UnrealCloudDDC/Jupiter/Implementation/ScyllaSessionManager.cs b/Engine/Source/Programs/UnrealCloudDDC/Jupiter/Implementation/ScyllaSessionManager.cs
--- a/Engine/Source/Programs/UnrealCloudDDC/Jupiter/Implementation/ScyllaSessionManager.cs
+++ b/Engine/Source/Programs/UnrealCloudDDC/Jupiter/Implementation/ScyllaSessionManager.cs
@@ -58,7 +58,7 @@ namespace Jupiter.Implementation
         [System.Diagnostics.CodeAnalysis.SuppressMessage("Style", "IDE0060:Remove unused parameter", Justification = "Useful to create from the tracer in case we want to use that context")]
         public static TelemetrySpan BuildScyllaSpan(this Tracer tracer, string spanName)
         {
-            return ScyllaTracer.StartActiveSpan(spanName).SetAttribute("type", "db").SetAttribute("service.name", ScyllaServiceName);
+            return ScyllaTracer.StartActiveSpan(spanName).SetAttribute("type", "db").SetAttribute("service.name", ScyllaServiceName).SetAttribute("operation.name", spanName);
         }
     }
 }
