UnrealCloudDDC - More tweaks to OpenTelemetry tracing in Dev

diff --git a/Engine/Source/Programs/UnrealCloudDDC/Jupiter.Common/BaseStartup.cs b/Engine/Source/Programs/UnrealCloudDDC/Jupiter.Common/BaseStartup.cs
--- a/Engine/Source/Programs/UnrealCloudDDC/Jupiter.Common/BaseStartup.cs
+++ b/Engine/Source/Programs/UnrealCloudDDC/Jupiter.Common/BaseStartup.cs
@@ -258,9 +258,10 @@ namespace Jupiter
 
                     builder.AddSource(tracerServiceName, "ScyllaDB");
 
-                    builder.SetResourceBuilder(ResourceBuilder.CreateDefault().AddService(tracerServiceName,
-                        serviceVersion: settings.CurrentValue.TracerServiceVersion,
-                        serviceInstanceId: Environment.MachineName));
+                    builder.SetResourceBuilder(ResourceBuilder.CreateEmpty()
+                        .AddService(tracerServiceName, serviceNamespace: "Jupiter", serviceVersion: settings.CurrentValue.TracerServiceVersion)
+                        .AddEnvironmentVariableDetector()
+                    );
                 });
             });
             services.Configure<OpenTelemetryLoggerOptions>(opt =>
