UnrealCloudDDC - Attempted to fix Scylla traces not showing up under our asp.net traces and tried to tweak the presentation of the http scopes

diff --git a/Engine/Source/Programs/UnrealCloudDDC/Jupiter.Common/BaseStartup.cs b/Engine/Source/Programs/UnrealCloudDDC/Jupiter.Common/BaseStartup.cs
--- a/Engine/Source/Programs/UnrealCloudDDC/Jupiter.Common/BaseStartup.cs
+++ b/Engine/Source/Programs/UnrealCloudDDC/Jupiter.Common/BaseStartup.cs
@@ -246,7 +246,20 @@ namespace Jupiter
 
             services.AddOpenTelemetryTracing(builder =>
             {
-                builder.AddHttpClientInstrumentation();
+                builder.AddHttpClientInstrumentation(options =>
+                {
+                    options.EnrichWithHttpRequestMessage = (activity, message) =>
+                    {
+                        activity.AddTag("service.name", activity.Source.Name + "-http-client");
+                        activity.AddTag("operation.name", "http-request");
+                    };
+
+                    options.EnrichWithHttpResponseMessage = (activity, message) =>
+                    {
+                        activity.AddTag("service.name", activity.Source.Name + "-http-client");
+                        activity.AddTag("operation.name", "http-response");
+                    };
+                });
                 builder.AddAspNetCoreInstrumentation();
 
                 builder.AddOtlpExporter();
diff --git a/Engine/Source/Programs/UnrealCloudDDC/Jupiter/Implementation/ScyllaSessionManager.cs b/Engine/Source/Programs/UnrealCloudDDC/Jupiter/Implementation/ScyllaSessionManager.cs
--- a/Engine/Source/Programs/UnrealCloudDDC/Jupiter/Implementation/ScyllaSessionManager.cs
+++ b/Engine/Source/Programs/UnrealCloudDDC/Jupiter/Implementation/ScyllaSessionManager.cs
@@ -58,7 +58,7 @@ namespace Jupiter.Implementation
         [System.Diagnostics.CodeAnalysis.SuppressMessage("Style", "IDE0060:Remove unused parameter", Justification = "Useful to create from the tracer in case we want to use that context")]
         public static TelemetrySpan BuildScyllaSpan(this Tracer tracer, string spanName)
         {
-            return ScyllaTracer.StartActiveSpan(spanName).SetAttribute("type", "db").SetAttribute("service.name", ScyllaServiceName).SetAttribute("operation.name", spanName);
+            return ScyllaTracer.StartActiveSpan(spanName, parentContext: Tracer.CurrentSpan.Context).SetAttribute("type", "db").SetAttribute("service.name", ScyllaServiceName).SetAttribute("operation.name", spanName);
         }
     }
 }
