UnrealCloudDDC - Adopted more standard OpenTelemetry options for controlling instrumentation and tweaked some spans to appear closer to the information we got from the datadog automatic tracing (which was more readable then OpenTelemetry out of the box)

diff --git a/Engine/Source/Programs/UnrealCloudDDC/Jupiter.Common/BaseProgram.cs b/Engine/Source/Programs/UnrealCloudDDC/Jupiter.Common/BaseProgram.cs
--- a/Engine/Source/Programs/UnrealCloudDDC/Jupiter.Common/BaseProgram.cs
+++ b/Engine/Source/Programs/UnrealCloudDDC/Jupiter.Common/BaseProgram.cs
@@ -6,9 +6,6 @@ using System.IO;
 using Microsoft.AspNetCore.Hosting;
 using Microsoft.Extensions.Configuration;
 using Microsoft.Extensions.Hosting;
-using OpenTelemetry.Trace;
-using OpenTracing;
-using OpenTracing.Util;
 using Serilog;
 using Serilog.Core;
 
diff --git a/Engine/Source/Programs/UnrealCloudDDC/Jupiter.Common/BaseStartup.cs b/Engine/Source/Programs/UnrealCloudDDC/Jupiter.Common/BaseStartup.cs
--- a/Engine/Source/Programs/UnrealCloudDDC/Jupiter.Common/BaseStartup.cs
+++ b/Engine/Source/Programs/UnrealCloudDDC/Jupiter.Common/BaseStartup.cs
@@ -31,6 +31,7 @@ using Newtonsoft.Json;
 using Newtonsoft.Json.Serialization;
 using Okta.AspNet.Abstractions;
 using Okta.AspNetCore;
+using OpenTelemetry.Exporter;
 using OpenTelemetry.Logs;
 using OpenTelemetry.Resources;
 using OpenTelemetry.Trace;
@@ -75,8 +76,6 @@ namespace Jupiter
             services.AddOptions<JupiterSettings>().Bind(Configuration.GetSection("Jupiter")).ValidateDataAnnotations();
             services.AddOptions<NamespaceSettings>().Bind(Configuration.GetSection("Namespaces")).ValidateDataAnnotations();
 
-            services.AddOptions<TracerSettings>().Bind(Configuration.GetSection("Tracer")).ValidateDataAnnotations();
-
             services.AddSingleton(typeof(INamespacePolicyResolver), typeof(NamespacePolicyResolver));
 
             // this is the same as invoke MvcBuilder.AddJsonOptions but with a service provider passed so we can use DI in the options creation
@@ -244,38 +243,31 @@ namespace Jupiter
             services.AddSingleton<IAuthorizationHandler, NamespaceAuthorizationHandler>();
             services.AddSingleton<IAuthorizationHandler, GlobalAuthorizationHandler>();
 
+            string otelServiceName = Configuration["OTEL_SERVICE_NAME"] ?? "unreal-cloud-ddc";
+            string otelServiceVersion = Configuration["OTEL_SERVICE_VERSION"];
             services.AddOpenTelemetryTracing(builder =>
             {
                 builder.AddHttpClientInstrumentation(options =>
                 {
                     options.EnrichWithHttpRequestMessage = (activity, message) =>
                     {
-                        activity.AddTag("service.name", activity.Source.Name + "-http-client");
+                        activity.AddTag("service.name", otelServiceName + "-http-client");
                         activity.AddTag("operation.name", "http-request");
-                    };
-
-                    options.EnrichWithHttpResponseMessage = (activity, message) =>
-                    {
-                        activity.AddTag("service.name", activity.Source.Name + "-http-client");
-                        activity.AddTag("operation.name", "http-response");
+                        string url = $"{message.Method} {message.Headers.Host}{message.RequestUri?.LocalPath}";
+                        activity.DisplayName = url;
                     };
                 });
                 builder.AddAspNetCoreInstrumentation();
 
                 builder.AddOtlpExporter();
 
-                builder.ConfigureBuilder((provider, _) =>
+                builder.ConfigureResource(resourceBuilder =>
                 {
-                    IOptionsMonitor<TracerSettings> settings = provider.GetService<IOptionsMonitor<TracerSettings>>()!;
-                    string tracerServiceName = settings.CurrentValue.TracerServiceName ?? "UnrealCloudDDC";
-
-                    builder.AddSource(tracerServiceName, "ScyllaDB");
-
-                    builder.SetResourceBuilder(ResourceBuilder.CreateEmpty()
-                        .AddService(tracerServiceName, serviceNamespace: "Jupiter", serviceVersion: settings.CurrentValue.TracerServiceVersion)
-                        .AddEnvironmentVariableDetector()
-                    );
+                    resourceBuilder.AddService("UnrealCloudDDC", serviceNamespace: "Jupiter", serviceVersion: otelServiceVersion);
+                    resourceBuilder.AddEnvironmentVariableDetector();
                 });
+
+                builder.AddSource("UnrealCloudDDC", "ScyllaDB");
             });
             services.Configure<OpenTelemetryLoggerOptions>(opt =>
             {
@@ -681,12 +673,6 @@ namespace Jupiter
         public string CurrentSite { get; set; } = "";
     }
 
-    public class TracerSettings
-    {
-        public string? TracerServiceName { get; set; } = null;
-        public string? TracerServiceVersion { get; set; } = null;
-    }
-
     public class NamespaceSettings
     {
         [System.Diagnostics.CodeAnalysis.SuppressMessage("Usage", "CA2227:Collection properties should be read only", Justification = "Used by the configuration system")]
diff --git a/Engine/Source/Programs/UnrealCloudDDC/Jupiter/Implementation/ScyllaSessionManager.cs b/Engine/Source/Programs/UnrealCloudDDC/Jupiter/Implementation/ScyllaSessionManager.cs
--- a/Engine/Source/Programs/UnrealCloudDDC/Jupiter/Implementation/ScyllaSessionManager.cs
+++ b/Engine/Source/Programs/UnrealCloudDDC/Jupiter/Implementation/ScyllaSessionManager.cs
@@ -58,7 +58,7 @@ namespace Jupiter.Implementation
         [System.Diagnostics.CodeAnalysis.SuppressMessage("Style", "IDE0060:Remove unused parameter", Justification = "Useful to create from the tracer in case we want to use that context")]
         public static TelemetrySpan BuildScyllaSpan(this Tracer tracer, string spanName)
         {
-            return ScyllaTracer.StartActiveSpan(spanName, parentContext: Tracer.CurrentSpan.Context).SetAttribute("type", "db").SetAttribute("service.name", ScyllaServiceName).SetAttribute("operation.name", spanName);
+            return ScyllaTracer.StartActiveSpan(spanName, parentContext: Tracer.CurrentSpan.Context).SetAttribute("type", "db").SetAttribute("operation.name", spanName);
         }
     }
 }
