UnrealCloudDDC - Added resource names to http calls to emulate datadog instrumentation. Also added service.name to scylla calls to correctly assign these traces to the scylla service

diff --git a/Engine/Source/Programs/UnrealCloudDDC/Jupiter/Implementation/ScyllaSessionManager.cs b/Engine/Source/Programs/UnrealCloudDDC/Jupiter/Implementation/ScyllaSessionManager.cs
--- a/Engine/Source/Programs/UnrealCloudDDC/Jupiter/Implementation/ScyllaSessionManager.cs
+++ b/Engine/Source/Programs/UnrealCloudDDC/Jupiter/Implementation/ScyllaSessionManager.cs
@@ -58,7 +58,7 @@ namespace Jupiter.Implementation
         [System.Diagnostics.CodeAnalysis.SuppressMessage("Style", "IDE0060:Remove unused parameter", Justification = "Useful to create from the tracer in case we want to use that context")]
         public static TelemetrySpan BuildScyllaSpan(this Tracer tracer, string spanName)
         {
-            return ScyllaTracer.StartActiveSpan(spanName, parentContext: Tracer.CurrentSpan.Context).SetAttribute("type", "db").SetAttribute("operation.name", spanName);
+            return ScyllaTracer.StartActiveSpan(spanName, parentContext: Tracer.CurrentSpan.Context).SetAttribute("type", "db").SetAttribute("operation.name", spanName).SetAttribute("service.name", ScyllaServiceName);
         }
     }
 }
