UnrealCloudDDC - Added operation.name tag to operations that were missing it, this is used by datadog to group operations together and is the display name of the span in the flame graph.

diff --git a/Engine/Source/Programs/UnrealCloudDDC/Jupiter.Common/BaseStartup.cs b/Engine/Source/Programs/UnrealCloudDDC/Jupiter.Common/BaseStartup.cs
--- a/Engine/Source/Programs/UnrealCloudDDC/Jupiter.Common/BaseStartup.cs
+++ b/Engine/Source/Programs/UnrealCloudDDC/Jupiter.Common/BaseStartup.cs
@@ -253,8 +253,11 @@ namespace Jupiter
                     {
                         activity.AddTag("service.name", otelServiceName + "-http-client");
                         activity.AddTag("operation.name", "http-request");
+
                         string url = $"{message.Method} {message.Headers.Host}{message.RequestUri?.LocalPath}";
                         activity.DisplayName = url;
+
+                        activity.AddTag("resource.name", url);
                     };
                 });
                 builder.AddAspNetCoreInstrumentation();
diff --git a/Engine/Source/Programs/UnrealCloudDDC/Jupiter.Common/Implementation/StreamingBodyPayload.cs b/Engine/Source/Programs/UnrealCloudDDC/Jupiter.Common/Implementation/StreamingBodyPayload.cs
--- a/Engine/Source/Programs/UnrealCloudDDC/Jupiter.Common/Implementation/StreamingBodyPayload.cs
+++ b/Engine/Source/Programs/UnrealCloudDDC/Jupiter.Common/Implementation/StreamingBodyPayload.cs
@@ -30,8 +30,9 @@ namespace Jupiter.Common.Implementation
 
         public static async Task<MemoryBufferedPayload> Create(Tracer tracer, Stream s)
         {
-            using TelemetrySpan scope = tracer.StartActiveSpan("payload.buffer");
-            scope.SetAttribute("bufferType", "Memory");
+            using TelemetrySpan scope = tracer.StartActiveSpan("payload.buffer")
+                .SetAttribute("operation.name", "payload.buffer")
+                .SetAttribute("bufferType", "Memory");
             MemoryBufferedPayload payload = new MemoryBufferedPayload(await s.ToByteArray());
 
             return payload;
@@ -68,8 +69,9 @@ namespace Jupiter.Common.Implementation
             FilesystemBufferedPayload payload = new FilesystemBufferedPayload();
 
             {
-                using TelemetrySpan? scope = tracer.StartActiveSpan("payload.buffer");
-                scope.SetAttribute("bufferType", "Filesystem");
+                using TelemetrySpan? scope = tracer.StartActiveSpan("payload.buffer")
+                    .SetAttribute("operation.name", "payload.buffer")
+                    .SetAttribute("bufferType", "Filesystem");
                 await using FileStream fs = payload._tempFile.OpenWrite();
                 await s.CopyToAsync(fs);
             }
diff --git a/Engine/Source/Programs/UnrealCloudDDC/Jupiter.Common/Utils/RequestUtil.cs b/Engine/Source/Programs/UnrealCloudDDC/Jupiter.Common/Utils/RequestUtil.cs
--- a/Engine/Source/Programs/UnrealCloudDDC/Jupiter.Common/Utils/RequestUtil.cs
+++ b/Engine/Source/Programs/UnrealCloudDDC/Jupiter.Common/Utils/RequestUtil.cs
@@ -21,8 +21,9 @@ namespace Jupiter
             }
 
             Tracer? tracer = request.HttpContext.RequestServices.GetService<Tracer>();
-            using TelemetrySpan? scope = tracer?.StartActiveSpan("readbody");
-            scope?.SetAttribute("Content-Length", contentLength.Value);
+            using TelemetrySpan? scope = tracer?.StartActiveSpan("readbody")
+                .SetAttribute("operation.name", "readbody")
+                .SetAttribute("Content-Length", contentLength.Value);
 
             await using MemoryStream ms = new MemoryStream((int)contentLength);
             DateTime readStart = DateTime.Now;
diff --git a/Engine/Source/Programs/UnrealCloudDDC/Jupiter/Controllers/ReferencesController.cs b/Engine/Source/Programs/UnrealCloudDDC/Jupiter/Controllers/ReferencesController.cs
--- a/Engine/Source/Programs/UnrealCloudDDC/Jupiter/Controllers/ReferencesController.cs
+++ b/Engine/Source/Programs/UnrealCloudDDC/Jupiter/Controllers/ReferencesController.cs
@@ -141,7 +141,7 @@ namespace Jupiter.Controllers
                     using ServerTimingMetricScoped? serverTimingScope = serverTiming?.CreateServerTimingMetricScope("body.write", "Time spent writing body");
 
                     long contentLength = blobContents.Length;
-                    using TelemetrySpan scope = _tracer.StartActiveSpan("body.write");
+                    using TelemetrySpan scope = _tracer.StartActiveSpan("body.write").SetAttribute("operation.name", "body.write");
                     scope.SetAttribute("content-length", contentLength);
                     const int BufferSize = 64 * 1024;
                     Stream outputStream = Response.Body;
@@ -217,7 +217,7 @@ namespace Jupiter.Controllers
                     {
                         byte[] blobMemory;
                         {
-                            using TelemetrySpan scope = _tracer.StartActiveSpan("json.readblob");
+                            using TelemetrySpan scope = _tracer.StartActiveSpan("json.readblob").SetAttribute("operation.name", "authorize");
                             blobMemory = await blob.Stream.ToByteArray();
                         }
                         CbObject cb = new CbObject(blobMemory);
diff --git a/Engine/Source/Programs/UnrealCloudDDC/Jupiter/Controllers/RequestHelper.cs b/Engine/Source/Programs/UnrealCloudDDC/Jupiter/Controllers/RequestHelper.cs
--- a/Engine/Source/Programs/UnrealCloudDDC/Jupiter/Controllers/RequestHelper.cs
+++ b/Engine/Source/Programs/UnrealCloudDDC/Jupiter/Controllers/RequestHelper.cs
@@ -1,5 +1,6 @@
 // Copyright Epic Games, Inc. All Rights Reserved.
 
+using System.Collections.Generic;
 using System.Security.Claims;
 using System.Threading.Tasks;
 using EpicGames.Horde.Storage;
@@ -29,7 +30,7 @@ public class RequestHelper
 
     public async Task<ActionResult?> HasAccessToNamespace(ClaimsPrincipal user, HttpRequest request, NamespaceId ns, AclAction[] aclActions)
     {
-        using TelemetrySpan _ = _tracer.StartActiveSpan("authorize");
+        using TelemetrySpan _ = _tracer.StartActiveSpan("authorize").SetAttribute("operation.name", "authorize");
         AuthorizationResult authorizationResult = await _authorizationService.AuthorizeAsync(user, new NamespaceAccessRequest
         {
             Namespace = ns,
@@ -67,7 +68,7 @@ public class RequestHelper
 
     public async Task<ActionResult?> HasAccessForGlobalOperations(ClaimsPrincipal user, AclAction[] aclActions)
     {
-        using TelemetrySpan _ = _tracer.StartActiveSpan("authorize");
+        using TelemetrySpan _ = _tracer.StartActiveSpan("authorize").SetAttribute("operation.name", "authorize");
         AuthorizationResult authorizationResult = await _authorizationService.AuthorizeAsync(user, new GlobalAccessRequest
         {
             Actions = aclActions
diff --git a/Engine/Source/Programs/UnrealCloudDDC/Jupiter/Implementation/Blob/BlobIndexConsistencyCheckService.cs b/Engine/Source/Programs/UnrealCloudDDC/Jupiter/Implementation/Blob/BlobIndexConsistencyCheckService.cs
--- a/Engine/Source/Programs/UnrealCloudDDC/Jupiter/Implementation/Blob/BlobIndexConsistencyCheckService.cs
+++ b/Engine/Source/Programs/UnrealCloudDDC/Jupiter/Implementation/Blob/BlobIndexConsistencyCheckService.cs
@@ -81,7 +81,7 @@ namespace Jupiter.Implementation
                         _logger.Information("Consistency check running on blob index, count of blobs processed so far: {CountOfBlobs}", countOfBlobsChecked);
                     }
 
-                    using TelemetrySpan scope = _tracer.StartActiveSpan("consistency_check.blob_index").SetAttribute("resource.name", $"{blobInfo.Namespace}.{blobInfo.BlobIdentifier}");
+                    using TelemetrySpan scope = _tracer.StartActiveSpan("consistency_check.blob_index").SetAttribute("resource.name", $"{blobInfo.Namespace}.{blobInfo.BlobIdentifier}").SetAttribute("operation.name", "consistency_check.blob_index");
 
                     bool issueFound = false;
                     bool deleted = false;
diff --git a/Engine/Source/Programs/UnrealCloudDDC/Jupiter/Implementation/Blob/BlobStoreConsistencyCheckService.cs b/Engine/Source/Programs/UnrealCloudDDC/Jupiter/Implementation/Blob/BlobStoreConsistencyCheckService.cs
--- a/Engine/Source/Programs/UnrealCloudDDC/Jupiter/Implementation/Blob/BlobStoreConsistencyCheckService.cs
+++ b/Engine/Source/Programs/UnrealCloudDDC/Jupiter/Implementation/Blob/BlobStoreConsistencyCheckService.cs
@@ -91,6 +91,7 @@ namespace Jupiter.Implementation
                     await foreach ((BlobIdentifier blob, DateTime lastModified) in blobStore.ListObjects(ns))
                     {
                         using TelemetrySpan scope = _tracer.StartActiveSpan("consistency_check.blob_store")
+                            .SetAttribute("operation.name", "consistency_check.blob_store")
                             .SetAttribute("resource.name", $"{ns}.{blob}")
                             .SetAttribute("BlobStore", blobStoreName);
 
diff --git a/Engine/Source/Programs/UnrealCloudDDC/Jupiter/Implementation/Blob/FileSystemStore.cs b/Engine/Source/Programs/UnrealCloudDDC/Jupiter/Implementation/Blob/FileSystemStore.cs
--- a/Engine/Source/Programs/UnrealCloudDDC/Jupiter/Implementation/Blob/FileSystemStore.cs
+++ b/Engine/Source/Programs/UnrealCloudDDC/Jupiter/Implementation/Blob/FileSystemStore.cs
@@ -217,7 +217,8 @@ namespace Jupiter.Implementation
         /// <returns></returns>
         public async Task<ulong> CleanupInternal(CancellationToken cancellationToken, int batchSize = 100000)
         {
-            using TelemetrySpan scope = _tracer.StartActiveSpan("gc.filesystem");
+            using TelemetrySpan scope = _tracer.StartActiveSpan("gc.filesystem")
+                .SetAttribute("operation.name", "gc.filesystem");
 
             ulong maxSizeBytes = _settings.CurrentValue.MaxSizeBytes;
             long triggerSize = (long) (maxSizeBytes * _settings.CurrentValue.TriggerThresholdPercentage);
diff --git a/Engine/Source/Programs/UnrealCloudDDC/Jupiter/Implementation/Blob/IBlobService.cs b/Engine/Source/Programs/UnrealCloudDDC/Jupiter/Implementation/Blob/IBlobService.cs
--- a/Engine/Source/Programs/UnrealCloudDDC/Jupiter/Implementation/Blob/IBlobService.cs
+++ b/Engine/Source/Programs/UnrealCloudDDC/Jupiter/Implementation/Blob/IBlobService.cs
@@ -123,7 +123,7 @@ public class BlobService : IBlobService
     {
         ContentHash blobHash;
         {
-            using TelemetrySpan _ = _tracer.StartActiveSpan("web.hash");
+            using TelemetrySpan _ = _tracer.StartActiveSpan("web.hash").SetAttribute("operation.name", "web.hash");
             blobHash = await BlobIdentifier.FromStream(content);
         }
 
@@ -138,6 +138,7 @@ public class BlobService : IBlobService
     public async Task<BlobIdentifier> PutObject(NamespaceId ns, IBufferedPayload payload, BlobIdentifier identifier)
     {
         using TelemetrySpan scope = _tracer.StartActiveSpan("put_blob")
+            .SetAttribute("operation.name", "put_blob")
             .SetAttribute("resource.name", identifier.ToString())
             .SetAttribute("Content-Length", payload.Length.ToString());
 
@@ -154,6 +155,7 @@ public class BlobService : IBlobService
     public async Task<BlobIdentifier> PutObject(NamespaceId ns, byte[] payload, BlobIdentifier identifier)
     {
         using TelemetrySpan scope = _tracer.StartActiveSpan("put_blob")
+            .SetAttribute("operation.name", "put_blob")
             .SetAttribute("resource.name", identifier.ToString())
             .SetAttribute("Content-Length", payload.Length.ToString())
             ;
@@ -170,6 +172,7 @@ public class BlobService : IBlobService
     public async Task<BlobIdentifier> PutObjectKnownHash(NamespaceId ns, IBufferedPayload content, BlobIdentifier identifier)
     {
         using TelemetrySpan scope = _tracer.StartActiveSpan("put_blob")
+            .SetAttribute("operation.name", "put_blob")
             .SetAttribute("resource.name", identifier.ToString())
             .SetAttribute("Content-Length", content.Length.ToString())
             ;
@@ -187,6 +190,7 @@ public class BlobService : IBlobService
         foreach (IBlobStore store in _blobStores)
         {
             using TelemetrySpan scope = _tracer.StartActiveSpan("put_blob_to_store")
+                .SetAttribute("operation.name", "put_blob_to_store")
                 .SetAttribute("resource.name", identifier.ToString())
                 .SetAttribute("store", store.GetType().ToString())
                 ;
@@ -243,6 +247,7 @@ public class BlobService : IBlobService
             }
 
             using TelemetrySpan scope = _tracer.StartActiveSpan("HierarchicalStore.GetObject")
+                .SetAttribute("operation.name", "HierarchicalStore.GetObject")
                 .SetAttribute("resource.name",  blob.ToString())
                 .SetAttribute("BlobStore", store.GetType().ToString())
                 .SetAttribute("ObjectFound", false.ToString())
@@ -292,7 +297,7 @@ public class BlobService : IBlobService
 
         if (numStoreMisses >= 1)
         {
-            using TelemetrySpan _ = _tracer.StartActiveSpan("HierarchicalStore.Populate");
+            using TelemetrySpan _ = _tracer.StartActiveSpan("HierarchicalStore.Populate").SetAttribute("operation.name", "HierarchicalStore.Populate");
             using ServerTimingMetricScoped? serverTimingScope = serverTiming?.CreateServerTimingMetricScope($"blob.populate", "Populating caches with blob contents");
 
             await using MemoryStream tempStream = new MemoryStream();
@@ -304,6 +309,7 @@ public class BlobService : IBlobService
             {
                 IBlobStore blobStore = _blobStores[i];
                 using TelemetrySpan scope = _tracer.StartActiveSpan("HierarchicalStore.PopulateStore")
+                    .SetAttribute("operation.name", "HierarchicalStore.PopulateStore")
                     .SetAttribute("BlobStore", blobStore.GetType().Name);
                 // Populate each store traversed that did not have the content found lower in the hierarchy
                 await blobStore.PutObject(ns, data, blob);
@@ -324,7 +330,7 @@ public class BlobService : IBlobService
             throw new NotSupportedException($"Replication is not allowed in namespace {ns}");
         }
         IServerTiming? serverTiming = _httpContextAccessor.HttpContext?.RequestServices.GetService<IServerTiming>();
-        using TelemetrySpan scope = _tracer.StartActiveSpan("HierarchicalStore.Replicate");
+        using TelemetrySpan scope = _tracer.StartActiveSpan("HierarchicalStore.Replicate").SetAttribute("operation.name", "HierarchicalStore.Replicate");
 
         using ServerTimingMetricScoped? serverTimingScope = serverTiming?.CreateServerTimingMetricScope("blob.replicate", "Replicating blob from remote instances");
 
@@ -407,6 +413,7 @@ public class BlobService : IBlobService
         if (useBlobIndex)
         {
             using TelemetrySpan scope = _tracer.StartActiveSpan("HierarchicalStore.ObjectExists")
+                .SetAttribute("operation.name", "HierarchicalStore.ObjectExists")
                 .SetAttribute("resource.name", blob.ToString())
                 .SetAttribute("BlobStore", "BlobIndex")
                 ;
@@ -445,6 +452,7 @@ public class BlobService : IBlobService
                 }
 
                 using TelemetrySpan scope = _tracer.StartActiveSpan("HierarchicalStore.ObjectExists")
+                    .SetAttribute("operation.name", "HierarchicalStore.ObjectExists")
                     .SetAttribute("resource.name", blob.ToString())
                     .SetAttribute("BlobStore", blobStoreName)
                     ;
@@ -465,6 +473,7 @@ public class BlobService : IBlobService
         IBlobStore store = _blobStores.Last();
 
         using TelemetrySpan scope = _tracer.StartActiveSpan("HierarchicalStore.ObjectExistsInRoot")
+            .SetAttribute("operation.name", "HierarchicalStore.ObjectExistsInRoot")
             .SetAttribute("resource.name", blob.ToString())
             .SetAttribute("BlobStore", store.GetType().Name)
             ;
@@ -492,6 +501,7 @@ public class BlobService : IBlobService
             try
             {
                 using TelemetrySpan scope = _tracer.StartActiveSpan("HierarchicalStore.DeleteObject")
+                    .SetAttribute("operation.name", "HierarchicalStore.DeleteObject")
                     .SetAttribute("resource.name", blob.ToString())
                     .SetAttribute("BlobStore", store.GetType().Name)
                     ;
@@ -528,6 +538,7 @@ public class BlobService : IBlobService
         foreach (IBlobStore store in _blobStores)
         {
             using TelemetrySpan scope = _tracer.StartActiveSpan("HierarchicalStore.DeleteNamespace")
+                .SetAttribute("operation.name", "HierarchicalStore.DeleteNamespace")
                 .SetAttribute("resource.name", ns.ToString())
                 .SetAttribute("BlobStore", store.GetType().Name)
                 ;
@@ -595,7 +606,7 @@ public class BlobService : IBlobService
 
     public async Task<BlobContents> GetObjects(NamespaceId ns, BlobIdentifier[] blobs)
     {
-        using TelemetrySpan _ = _tracer.StartActiveSpan("blob.combine");
+        using TelemetrySpan _ = _tracer.StartActiveSpan("blob.combine").SetAttribute("operation.name", "blob.combine");
         Task<BlobContents>[] tasks = new Task<BlobContents>[blobs.Length];
         for (int i = 0; i < blobs.Length; i++)
         {
@@ -639,7 +650,7 @@ public static class BlobServiceExtensions
         {
             ContentHash blobHash;
             {
-                using TelemetrySpan _ = tracer.StartActiveSpan("web.hash");
+                using TelemetrySpan _ = tracer.StartActiveSpan("web.hash").SetAttribute("operation.name", "web.hash");
                 blobHash = await BlobIdentifier.FromStream(decompressedStream);
             }
 
@@ -648,7 +659,7 @@ public static class BlobServiceExtensions
 
         BlobIdentifier identifierCompressedPayload;
         {
-            using TelemetrySpan _ = tracer.StartActiveSpan("web.hash");
+            using TelemetrySpan _ = tracer.StartActiveSpan("web.hash").SetAttribute("operation.name", "web.hash");
             await using Stream hashStream = payload.GetStream();
             identifierCompressedPayload = await BlobIdentifier.FromStream(hashStream);
         }
@@ -694,7 +705,7 @@ public static class BlobServiceExtensions
         }
 
         // chunked content, combine the chunks into a single stream
-        using TelemetrySpan _ = tracer.StartActiveSpan("blob.combine");
+        using TelemetrySpan _ = tracer.StartActiveSpan("blob.combine").SetAttribute("operation.name", "blob.combine");
         Task<BlobContents>[] tasks = new Task<BlobContents>[chunks.Length];
         for (int i = 0; i < chunks.Length; i++)
         {
diff --git a/Engine/Source/Programs/UnrealCloudDDC/Jupiter/Implementation/ContentId/CompressedBufferUtils.cs b/Engine/Source/Programs/UnrealCloudDDC/Jupiter/Implementation/ContentId/CompressedBufferUtils.cs
--- a/Engine/Source/Programs/UnrealCloudDDC/Jupiter/Implementation/ContentId/CompressedBufferUtils.cs
+++ b/Engine/Source/Programs/UnrealCloudDDC/Jupiter/Implementation/ContentId/CompressedBufferUtils.cs
@@ -229,7 +229,7 @@ namespace Jupiter.Implementation
             }
 
             {
-                using TelemetrySpan _ = _tracer.StartActiveSpan("web.hash");
+                using TelemetrySpan _ = _tracer.StartActiveSpan("web.hash").SetAttribute("operation.name", "web.hash");
 
                 // only read the first 20 bytes of the hash field as IoHashes are 20 bytes and not 32 bytes
                 byte[] slicedHash = new byte[20];
diff --git a/Engine/Source/Programs/UnrealCloudDDC/Jupiter/Implementation/ContentId/ScyllaContentIdStore.cs b/Engine/Source/Programs/UnrealCloudDDC/Jupiter/Implementation/ContentId/ScyllaContentIdStore.cs
--- a/Engine/Source/Programs/UnrealCloudDDC/Jupiter/Implementation/ContentId/ScyllaContentIdStore.cs
+++ b/Engine/Source/Programs/UnrealCloudDDC/Jupiter/Implementation/ContentId/ScyllaContentIdStore.cs
@@ -64,7 +64,7 @@ namespace Jupiter.Implementation
                         BlobIdentifier[] blobs = resolvedContentId.Chunks.Select(b => b.AsBlobIdentifier()).ToArray();
 
                         {
-                            using TelemetrySpan _ = _tracer.StartActiveSpan("ScyllaContentIdStore.FindMissingBlobs");
+                            using TelemetrySpan _ = _tracer.StartActiveSpan("ScyllaContentIdStore.FindMissingBlobs").SetAttribute("operation.name", "ScyllaContentIdStore.FindMissingBlobs");
 
                             BlobIdentifier[] missingBlobs = await _blobStore.FilterOutKnownBlobs(ns, blobs);
                             if (missingBlobs.Length == 0)
@@ -88,7 +88,7 @@ namespace Jupiter.Implementation
                         BlobIdentifier[] blobs = resolvedContentId.Chunks.Select(b => b.AsBlobIdentifier()).ToArray();
 
                         {
-                            using TelemetrySpan _ = _tracer.StartActiveSpan("ScyllaContentIdStore.FindMissingBlobs");
+                            using TelemetrySpan _ = _tracer.StartActiveSpan("ScyllaContentIdStore.FindMissingBlobs").SetAttribute("operation.name", "ScyllaContentIdStore.FindMissingBlobs");
 
                             BlobIdentifier[] missingBlobs = await _blobStore.FilterOutKnownBlobs(ns, blobs);
                             if (missingBlobs.Length == 0)
diff --git a/Engine/Source/Programs/UnrealCloudDDC/Jupiter/Implementation/GC/CacheCleanup.cs b/Engine/Source/Programs/UnrealCloudDDC/Jupiter/Implementation/GC/CacheCleanup.cs
--- a/Engine/Source/Programs/UnrealCloudDDC/Jupiter/Implementation/GC/CacheCleanup.cs
+++ b/Engine/Source/Programs/UnrealCloudDDC/Jupiter/Implementation/GC/CacheCleanup.cs
@@ -86,7 +86,10 @@ namespace Jupiter.Implementation
                     _logger.Information(
                         "Attempting to delete object {Namespace} {Bucket} {Name} as it was last updated {LastAccessTime} which is older then {CutoffTime}",
                         ns, bucket, name, lastAccessTime, cutoffTime);
-                    using TelemetrySpan scope = _tracer.StartActiveSpan("gc.ref").SetAttribute("resource.name", $"{ns}:{bucket}.{name}").SetAttribute("namespace", ns.ToString());
+                    using TelemetrySpan scope = _tracer.StartActiveSpan("gc.ref")
+                        .SetAttribute("operation.name", "gc.ref")
+                        .SetAttribute("resource.name", $"{ns}:{bucket}.{name}")
+                        .SetAttribute("namespace", ns.ToString());
                     // delete the old record from the ref refs
 
                     bool storeDelete = false;
diff --git a/Engine/Source/Programs/UnrealCloudDDC/Jupiter/Implementation/GC/OrphanBlobCleanupRefs.cs b/Engine/Source/Programs/UnrealCloudDDC/Jupiter/Implementation/GC/OrphanBlobCleanupRefs.cs
--- a/Engine/Source/Programs/UnrealCloudDDC/Jupiter/Implementation/GC/OrphanBlobCleanupRefs.cs
+++ b/Engine/Source/Programs/UnrealCloudDDC/Jupiter/Implementation/GC/OrphanBlobCleanupRefs.cs
@@ -115,7 +115,9 @@ namespace Jupiter.Implementation
         private async Task<bool> GCBlob(string storagePool, List<NamespaceId> namespacesThatSharePool, BlobIdentifier blob, DateTime lastModifiedTime, CancellationToken cancellationToken)
         {
             string storagePoolName = string.IsNullOrEmpty(storagePool) ? "default" : storagePool; 
-            using TelemetrySpan removeBlobScope = _tracer.StartActiveSpan("gc.blob").SetAttribute("resource.name", $"{storagePoolName}.{blob}");
+            using TelemetrySpan removeBlobScope = _tracer.StartActiveSpan("gc.blob")
+                .SetAttribute("operation.name", "gc.blob")
+                .SetAttribute("resource.name", $"{storagePoolName}.{blob}");
 
             bool found = false;
 
diff --git a/Engine/Source/Programs/UnrealCloudDDC/Jupiter/Implementation/IReferenceResolver.cs b/Engine/Source/Programs/UnrealCloudDDC/Jupiter/Implementation/IReferenceResolver.cs
--- a/Engine/Source/Programs/UnrealCloudDDC/Jupiter/Implementation/IReferenceResolver.cs
+++ b/Engine/Source/Programs/UnrealCloudDDC/Jupiter/Implementation/IReferenceResolver.cs
@@ -266,7 +266,9 @@ namespace Jupiter.Implementation
 
         private async Task<(ContentId, BlobIdentifier[]?)> ResolveContentId(NamespaceId ns, ContentId contentId)
         {
-            using TelemetrySpan scope = _tracer.StartActiveSpan("ReferenceResolver.ResolveContentId").SetAttribute("resource.name", contentId.ToString());
+            using TelemetrySpan scope = _tracer.StartActiveSpan("ReferenceResolver.ResolveContentId")
+                .SetAttribute("operation.name", "ReferenceResolver.ResolveContentId")
+                .SetAttribute("resource.name", contentId.ToString());
             BlobIdentifier[]? resolvedBlobs = await _contentIdStore.Resolve(ns, contentId);
             return (contentId, resolvedBlobs);
         }
diff --git a/Engine/Source/Programs/UnrealCloudDDC/Jupiter/Implementation/LastAccess/LastAccessServiceReferences.cs b/Engine/Source/Programs/UnrealCloudDDC/Jupiter/Implementation/LastAccess/LastAccessServiceReferences.cs
--- a/Engine/Source/Programs/UnrealCloudDDC/Jupiter/Implementation/LastAccess/LastAccessServiceReferences.cs
+++ b/Engine/Source/Programs/UnrealCloudDDC/Jupiter/Implementation/LastAccess/LastAccessServiceReferences.cs
@@ -81,7 +81,9 @@ namespace Jupiter.Implementation
                     continue;
                 }
 
-                using TelemetrySpan scope = _tracer.StartActiveSpan("lastAccess.update").SetAttribute("resource.name", $"{record.Namespace}:{record.Bucket}.{record.Key}");
+                using TelemetrySpan scope = _tracer.StartActiveSpan("lastAccess.update")
+                    .SetAttribute("operation.name", "lastAccess.update")
+                    .SetAttribute("resource.name", $"{record.Namespace}:{record.Bucket}.{record.Key}");
                 _logger.Debug("Updating last access time to {LastAccessTime} for {Record}", lastAccessTime, record);
                 await _referencesStore.UpdateLastAccessTime(record.Namespace, record.Bucket, record.Key, lastAccessTime);
                 // delay 10ms between each record to distribute the load more evenly for the db
diff --git a/Engine/Source/Programs/UnrealCloudDDC/Jupiter/Implementation/LastAccess/LastAccessTracker.cs b/Engine/Source/Programs/UnrealCloudDDC/Jupiter/Implementation/LastAccess/LastAccessTracker.cs
--- a/Engine/Source/Programs/UnrealCloudDDC/Jupiter/Implementation/LastAccess/LastAccessTracker.cs
+++ b/Engine/Source/Programs/UnrealCloudDDC/Jupiter/Implementation/LastAccess/LastAccessTracker.cs
@@ -67,7 +67,7 @@ namespace Jupiter.Implementation
 
             return Task.Run(() =>
             {
-                using TelemetrySpan _ = _tracer.StartActiveSpan("lastAccessTracker.track");
+                using TelemetrySpan _ = _tracer.StartActiveSpan("lastAccessTracker.track").SetAttribute("operation.name", "lastAccessTracker.track");
                 try
                 {
                     _rwLock.AcquireReaderLock(-1);
diff --git a/Engine/Source/Programs/UnrealCloudDDC/Jupiter/Implementation/ObjectService.cs b/Engine/Source/Programs/UnrealCloudDDC/Jupiter/Implementation/ObjectService.cs
--- a/Engine/Source/Programs/UnrealCloudDDC/Jupiter/Implementation/ObjectService.cs
+++ b/Engine/Source/Programs/UnrealCloudDDC/Jupiter/Implementation/ObjectService.cs
@@ -190,7 +190,7 @@ namespace Jupiter.Implementation
             bool hasReferences = HasAttachments(payload);
             if (hasReferences)
             {
-                using TelemetrySpan _ = _tracer.StartActiveSpan("ObjectService.ResolveReferences");
+                using TelemetrySpan _ = _tracer.StartActiveSpan("ObjectService.ResolveReferences").SetAttribute("operation.name", "ObjectService.ResolveReferences");
                 try
                 {
                     IAsyncEnumerable<BlobIdentifier> references = _referenceResolver.GetReferencedBlobs(ns, payload);
diff --git a/Engine/Source/Programs/UnrealCloudDDC/Jupiter/Implementation/Replication/RefsReplicator.cs b/Engine/Source/Programs/UnrealCloudDDC/Jupiter/Implementation/Replication/RefsReplicator.cs
--- a/Engine/Source/Programs/UnrealCloudDDC/Jupiter/Implementation/Replication/RefsReplicator.cs
+++ b/Engine/Source/Programs/UnrealCloudDDC/Jupiter/Implementation/Replication/RefsReplicator.cs
@@ -309,8 +309,9 @@ namespace Jupiter.Implementation
             {
                 try
                 {
-                    using TelemetrySpan? scope = _tracer.StartActiveSpan("replicator.replicate_op_snapshot");
-                    scope.SetAttribute("resource.name", $"{ns}.{snapshotLiveObject.Blob}");
+                    using TelemetrySpan? scope = _tracer.StartActiveSpan("replicator.replicate_op_snapshot")
+                        .SetAttribute("operation.name", "replicator.replicate_op_snapshot")
+                        .SetAttribute("resource.name", $"{ns}.{snapshotLiveObject.Blob}");
                     Interlocked.Increment(ref countOfObjectsCurrentlyReplicating);
                     LogReplicationHeartbeat(countOfObjectsCurrentlyReplicating);
 
@@ -376,6 +377,7 @@ namespace Jupiter.Implementation
             await GetRefEvents(ns, lastBucket, lastEvent, replicationToken).ParallelForEachAsync(async (ReplicationLogEvent @event) =>
             {
                 using TelemetrySpan? scope = _tracer.StartActiveSpan("replicator.replicate_op_incremental")
+                    .SetAttribute("operation.name", "replicator.replicate_op_incremental")
                     .SetAttribute("resource.name", $"[REDACTED_EMAIL]}[REDACTED_EMAIL]}")
                     .SetAttribute("time-bucket", @event.Timestamp.ToString(CultureInfo.InvariantCulture));
 
@@ -454,7 +456,9 @@ namespace Jupiter.Implementation
 
         private async Task<bool> ReplicateOp(NamespaceId ns, BlobIdentifier objectToReplicate, CancellationToken cancellationToken)
         {
-            using TelemetrySpan scope = _tracer.StartActiveSpan("replicator.replicate_op").SetAttribute("resource.name", $"{ns}.{objectToReplicate}");
+            using TelemetrySpan scope = _tracer.StartActiveSpan("replicator.replicate_op")
+                .SetAttribute("operation.name", "replicator.replicate_op")
+                .SetAttribute("resource.name", $"{ns}.{objectToReplicate}");
 
             _logger.Information("Attempting to replicate object {Blob} in {Namespace}.", objectToReplicate, ns);
 
