UnrealCloudDDC - Update changelog

diff --git a/Engine/Source/Programs/UnrealCloudDDC/changelog.md b/Engine/Source/Programs/UnrealCloudDDC/changelog.md
--- a/Engine/Source/Programs/UnrealCloudDDC/changelog.md
+++ b/Engine/Source/Programs/UnrealCloudDDC/changelog.md
@@ -1,5 +1,7 @@
 # Unreleased
 * Improved handling of cassandra timeouts during ref GC.
+* Added deletion of invalid references during blob GC - this reduces the size for entries in the blob_index table.
+* Using OpenTelemetry for tracing - still supports datadog traces but this can now be forwarded to any OpenTelemetry compatible service.
 
 # 0.1.0
 * First public release of UnrealCloudDDC
