UnrealCloudDDC - Replace _ with - for azure blob storage names (as _ is not allowed)

diff --git a/Engine/Source/Programs/UnrealCloudDDC/Jupiter/Implementation/Blob/AzureBlobStore.cs b/Engine/Source/Programs/UnrealCloudDDC/Jupiter/Implementation/Blob/AzureBlobStore.cs
--- a/Engine/Source/Programs/UnrealCloudDDC/Jupiter/Implementation/Blob/AzureBlobStore.cs
+++ b/Engine/Source/Programs/UnrealCloudDDC/Jupiter/Implementation/Blob/AzureBlobStore.cs
@@ -187,7 +187,7 @@ namespace Jupiter.Implementation
 
         private static string SanitizeNamespace(NamespaceId ns)
         {
-            return ns.ToString().Replace(".", "-", StringComparison.OrdinalIgnoreCase).ToLower();
+            return ns.ToString().Replace(".", "-", StringComparison.OrdinalIgnoreCase).Replace("_", "-", StringComparison.OrdinalIgnoreCase).ToLower();
         }
     }
 }
diff --git a/Engine/Source/Programs/UnrealCloudDDC/changelog.md b/Engine/Source/Programs/UnrealCloudDDC/changelog.md
--- a/Engine/Source/Programs/UnrealCloudDDC/changelog.md
+++ b/Engine/Source/Programs/UnrealCloudDDC/changelog.md
@@ -2,6 +2,7 @@
 * Improved handling of cassandra timeouts during ref GC.
 * Added deletion of invalid references during blob GC - this reduces the size for entries in the blob_index table.
 * Using OpenTelemetry for tracing - still supports datadog traces but this can now be forwarded to any OpenTelemetry compatible service.
+* Fixed issue in Azure Blob storage with namespace containing _
 
 # 0.1.0
 * First public release of UnrealCloudDDC
