UnrealCloudDDC - Added optional method of doing full table scan that scans each partition in order based on whitepapers from Scylla. This method can easily be adopted to scan multiple paritions in parallell if we want to complete the scan quicker (at the risk of adding more load to the cluster). For now we just scan a single partition at a time.

diff --git a/Engine/Source/Programs/UnrealCloudDDC/Jupiter/Implementation/Objects/ScyllaReferencesStore.cs b/Engine/Source/Programs/UnrealCloudDDC/Jupiter/Implementation/Objects/ScyllaReferencesStore.cs
--- a/Engine/Source/Programs/UnrealCloudDDC/Jupiter/Implementation/Objects/ScyllaReferencesStore.cs
+++ b/Engine/Source/Programs/UnrealCloudDDC/Jupiter/Implementation/Objects/ScyllaReferencesStore.cs
@@ -22,6 +22,7 @@ namespace Jupiter.Implementation
         private readonly ILogger _logger;
         private readonly PreparedStatement _getObjectsStatement;
         private readonly PreparedStatement _getNamespacesStatement;
+        private readonly PreparedStatement _getObjectsForPartitionRangeStatement;
 
         public ScyllaReferencesStore(IScyllaSessionManager scyllaSessionManager, IOptionsMonitor<ScyllaSettings> settings, Tracer tracer, ILogger<ScyllaReferencesStore> logger)
         {
@@ -55,6 +56,8 @@ namespace Jupiter.Implementation
             string cqlOptions = scyllaSessionManager.IsScylla ? "BYPASS CACHE" : "";
             _getObjectsStatement = _session.Prepare($"SELECT bucket, name, last_access_time FROM objects WHERE namespace = ? ALLOW FILTERING {cqlOptions}");
             _getNamespacesStatement = _session.Prepare("SELECT DISTINCT namespace FROM buckets");
+
+            _getObjectsForPartitionRangeStatement = _session.Prepare($"SELECT namespace, bucket, name, last_access_time FROM objects WHERE token(namespace, bucket, name) >= ? AND token(namespace, bucket, name) <= ? {cqlOptions}");
         }
 
         public async Task<ObjectRecord> Get(NamespaceId ns, BucketId bucket, IoHashKey name, IReferencesStore.FieldFlags flags)
@@ -125,62 +128,112 @@ namespace Jupiter.Implementation
         public async IAsyncEnumerable<(BucketId, IoHashKey, DateTime)> GetRecords(NamespaceId ns)
         {
             using TelemetrySpan scope = _tracer.BuildScyllaSpan("scylla.get_records");
-            const int MaxRetryAttempts = 3;
-            RowSet rowSet = await _session.ExecuteAsync(_getObjectsStatement.Bind(ns.ToString()));
 
-            do
+            if (_settings.CurrentValue.UsePerShardScanning)
             {
-                int countOfRows = rowSet.GetAvailableWithoutFetching();
-                IEnumerable<Row> localRows = rowSet.Take(countOfRows);
-                Task prefetchTask = rowSet.FetchMoreResultsAsync();
+                IAsyncEnumerable<(NamespaceId, BucketId, IoHashKey, DateTime)> enumerable = GetRecordsPerShard();
 
-                foreach (Row row in localRows)
+                await foreach ((NamespaceId, BucketId, IoHashKey, DateTime) record in enumerable)
                 {
-                    string bucket = row.GetValue<string>("bucket");
-                    string name = row.GetValue<string>("name");
-                    DateTime? lastAccessTime = row.GetValue<DateTime?>("last_access_time");
-
-                    // skip any names that are not conformant to io hash
-                    if (name.Length != 40)
+                    if (!record.Item1.Equals(ns))
                     {
                         continue;
                     }
-
-                    // if last access time is missing we treat it as being very old
-                    lastAccessTime ??= DateTime.MinValue;
-                    yield return (new BucketId(bucket), new IoHashKey(name), lastAccessTime.Value);
+                    yield return (record.Item2, record.Item3, record.Item4);
                 }
+            }
+            else
+            {
+                const int MaxRetryAttempts = 3;
+                RowSet rowSet = await _session.ExecuteAsync(_getObjectsStatement.Bind(ns.ToString()));
 
-                int retryAttempts = 0;
-                Exception? timeoutException = null;
-                while (retryAttempts < MaxRetryAttempts)
+                do
                 {
-                    try
+                    int countOfRows = rowSet.GetAvailableWithoutFetching();
+                    IEnumerable<Row> localRows = rowSet.Take(countOfRows);
+                    Task prefetchTask = rowSet.FetchMoreResultsAsync();
+
+                    foreach (Row row in localRows)
                     {
-                        await prefetchTask;
-                        timeoutException = null;
-                        break;
+                        string bucket = row.GetValue<string>("bucket");
+                        string name = row.GetValue<string>("name");
+                        DateTime? lastAccessTime = row.GetValue<DateTime?>("last_access_time");
+
+                        // skip any names that are not conformant to io hash
+                        if (name.Length != 40)
+                        {
+                            continue;
+                        }
+
+                        // if last access time is missing we treat it as being very old
+                        lastAccessTime ??= DateTime.MinValue;
+                        yield return (new BucketId(bucket), new IoHashKey(name), lastAccessTime.Value);
                     }
-                    catch (ReadTimeoutException e)
+
+                    int retryAttempts = 0;
+                    Exception? timeoutException = null;
+                    while (retryAttempts < MaxRetryAttempts)
                     {
-                        retryAttempts += 1;
-                        _logger.LogWarning(
-                            "Cassandra read timeouts, waiting a while and then retrying. Attempt {Attempts} .",
+                        try
+                        {
+                            await prefetchTask;
+                            timeoutException = null;
+                            break;
+                        }
+                        catch (ReadTimeoutException e)
+                        {
+                            retryAttempts += 1;
+                            _logger.LogWarning(
+                                "Cassandra read timeouts, waiting a while and then retrying. Attempt {Attempts} .",
+                                retryAttempts);
+                            // wait 10 seconds and try again as the Db is under heavy load right now
+                            await Task.Delay(TimeSpan.FromSeconds(10));
+                            timeoutException = e;
+                        }
+                    }
+
+                    if (timeoutException != null)
+                    {
+                        _logger.LogWarning("Cassandra read timeouts, attempted {Attempts} attempts now we give up.",
                             retryAttempts);
-                        // wait 10 seconds and try again as the Db is under heavy load right now
-                        await Task.Delay(TimeSpan.FromSeconds(10));
-                        timeoutException = e;
+                        // we have failed to many times, rethrow the exception and abort to avoid stalling here for ever
+                        throw timeoutException;
                     }
-                }
+                } while (!rowSet.IsFullyFetched);
+            }
+        }
 
-                if (timeoutException != null)
+        /// <summary>
+        /// This implements a more efficient scanning where we fetch objects based on which shard it is in. It scans the entire database and thus returns all namespaces.
+        /// See https://www.scylladb.com/2017/03/28/parallel-efficient-full-table-scan-scylla/
+        /// </summary>
+        /// <returns></returns>
+        private async IAsyncEnumerable<(NamespaceId, BucketId, IoHashKey, DateTime)> GetRecordsPerShard()
+        {
+            using TelemetrySpan scope = _tracer.BuildScyllaSpan("scylla.get_records_per_shard");
+            
+            foreach ((long, long) range in ScyllaUtils.GetTableRanges(_settings.CurrentValue.CountOfNodes, _settings.CurrentValue.CountOfCoresPerNode, 3))
+            {
+                RowSet rowSet = await _session.ExecuteAsync(_getObjectsForPartitionRangeStatement.Bind(range.Item1, range.Item2));
+                
+                foreach (Row row in rowSet)
                 {
-                    _logger.LogWarning("Cassandra read timeouts, attempted {Attempts} attempts now we give up.",
-                        retryAttempts);
-                    // we have failed to many times, rethrow the exception and abort to avoid stalling here for ever
-                    throw timeoutException;
+                    string ns = row.GetValue<string>("namespace");
+                    string bucket = row.GetValue<string>("bucket");
+                    string name = row.GetValue<string>("name");
+                    DateTime? lastAccessTime = row.GetValue<DateTime?>("last_access_time");
+
+                    // skip any names that are not conformant to io hash
+                    if (name.Length != 40)
+                    {
+                        continue;
+                    }
+
+                    // if last access time is missing we treat it as being very old
+                    lastAccessTime ??= DateTime.MinValue;
+                    yield return (new NamespaceId(ns), new BucketId(bucket), new IoHashKey(name), lastAccessTime.Value);
                 }
-            } while (!rowSet.IsFullyFetched);
+            }
         }
 
         public async IAsyncEnumerable<NamespaceId> GetNamespaces()
diff --git a/Engine/Source/Programs/UnrealCloudDDC/Jupiter/Implementation/Objects/ScyllaSettings.cs b/Engine/Source/Programs/UnrealCloudDDC/Jupiter/Implementation/Objects/ScyllaSettings.cs
--- a/Engine/Source/Programs/UnrealCloudDDC/Jupiter/Implementation/Objects/ScyllaSettings.cs
+++ b/Engine/Source/Programs/UnrealCloudDDC/Jupiter/Implementation/Objects/ScyllaSettings.cs
@@ -6,7 +6,7 @@ using System.ComponentModel.DataAnnotations;
 
 namespace Jupiter.Implementation;
 
-public class ScyllaSettings
+public class ScyllaSettings : IValidatableObject
 {
     public string ConnectionString { get; set; } = "Contact Points=localhost,scylla";
     public long InlineBlobMaxSize { get; set; } = 32 * 1024; // default to 32 kb blobs max
@@ -58,4 +58,32 @@ public class ScyllaSettings
     /// </summary>
     public bool UseSSL { get; set; } = true;
 
+    /// <summary>
+    /// Enable to scan database per shard, this requires setting more information about your cluster
+    /// </summary>
+    public bool UsePerShardScanning { get; set; } = false;
+
+    /// <summary>
+    /// Count of cores each nodes has, only used when UsePerShardScanning is set
+    /// </summary>
+    [Range(1, 128)]
+    public uint CountOfCoresPerNode { get; set; } = 1;
+
+    /// <summary>
+    /// Count of nodes in the DC, only used when UsePerShardScanning is set
+    /// </summary>
+    [Range(1, 64)]
+    public uint CountOfNodes { get; set; } = 1;
+
+    public IEnumerable<ValidationResult> Validate(ValidationContext validationContext)
+    {
+        List<ValidationResult> results = new List<ValidationResult>();
+
+        if (UsePerShardScanning)
+        {
+            Validator.TryValidateProperty(CountOfNodes, new ValidationContext(this, null, null) { MemberName = nameof(CountOfNodes) }, results);
+            Validator.TryValidateProperty(CountOfCoresPerNode, new ValidationContext(this, null, null) { MemberName = nameof(CountOfCoresPerNode) }, results);
+        }
+        return results;
+    }
 }
diff --git a/Engine/Source/Programs/UnrealCloudDDC/Jupiter/Implementation/ScyllaSessionManager.cs b/Engine/Source/Programs/UnrealCloudDDC/Jupiter/Implementation/ScyllaSessionManager.cs
--- a/Engine/Source/Programs/UnrealCloudDDC/Jupiter/Implementation/ScyllaSessionManager.cs
+++ b/Engine/Source/Programs/UnrealCloudDDC/Jupiter/Implementation/ScyllaSessionManager.cs
@@ -1,5 +1,7 @@
 ﻿// Copyright Epic Games, Inc. All Rights Reserved.
 
+using System;
+using System.Collections.Generic;
 using Cassandra;
 using OpenTelemetry.Trace;
 
@@ -49,6 +51,31 @@ namespace Jupiter.Implementation
         public bool IsCassandra { get; }
     }
 
+    public static class ScyllaUtils
+    {
+        public static IEnumerable<(long, long)> GetTableRanges(uint nodesInCluster, uint coresPerNode, uint smudgeFactor)
+        {
+            ulong parallelQueries = nodesInCluster * coresPerNode * smudgeFactor;
+            ulong countOfRanges = parallelQueries * 100;
+            ulong maxSize = ulong.MaxValue;
+            ulong rangeSize = maxSize / countOfRanges;
+
+            long start = long.MinValue;
+            long end = 0;
+            for (ulong i = 0; i < countOfRanges; i++)
+            {
+                end = start + (long)rangeSize;
+                if (start > 0 && end < 0)
+                {
+                    end = long.MaxValue;
+                }
+                yield return (start, end);
+
+                start = end + 1;
+            }
+        }
+    }
+
     public static class ScyllaTraceExtensions
     {
         private static readonly Tracer ScyllaTracer = TracerProvider.Default.GetTracer(ScyllaServiceName);
