UnrealCloudDDC - FilterOutKnownBlobs now running in parallell for normal enumerables and not just async enumeranbles.

diff --git a/Engine/Source/Programs/UnrealCloudDDC/Jupiter/Implementation/Blob/IBlobService.cs b/Engine/Source/Programs/UnrealCloudDDC/Jupiter/Implementation/Blob/IBlobService.cs
--- a/Engine/Source/Programs/UnrealCloudDDC/Jupiter/Implementation/Blob/IBlobService.cs
+++ b/Engine/Source/Programs/UnrealCloudDDC/Jupiter/Implementation/Blob/IBlobService.cs
@@ -52,7 +52,7 @@ public interface IBlobService
     Task DeleteNamespace(NamespaceId ns);
 
     IAsyncEnumerable<(BlobIdentifier,DateTime)> ListObjects(NamespaceId ns);
-    Task<BlobIdentifier[]> FilterOutKnownBlobs(NamespaceId ns, BlobIdentifier[] blobs);
+    Task<BlobIdentifier[]> FilterOutKnownBlobs(NamespaceId ns, IEnumerable<BlobIdentifier> blobs);
     Task<BlobIdentifier[]> FilterOutKnownBlobs(NamespaceId ns, IAsyncEnumerable<BlobIdentifier> blobs);
     Task<BlobContents> GetObjects(NamespaceId ns, BlobIdentifier[] refRequestBlobReferences);
 
@@ -567,12 +567,33 @@ public class BlobService : IBlobService
         return _blobStores.Last().ListObjects(ns);
     }
 
-    public async Task<BlobIdentifier[]> FilterOutKnownBlobs(NamespaceId ns, BlobIdentifier[] blobs)
+    public async Task<BlobIdentifier[]> FilterOutKnownBlobs(NamespaceId ns, IEnumerable<BlobIdentifier> blobs)
     {
-        var tasks = blobs.Select(async blobIdentifier => new { BlobIdentifier = blobIdentifier, Exists = await Exists(ns, blobIdentifier) });
-        var blobResults = await Task.WhenAll(tasks);
-        IEnumerable<BlobIdentifier> filteredBlobs = blobResults.Where(ac => !ac.Exists).Select(ac => ac.BlobIdentifier);
-        return filteredBlobs.ToArray();
+        ConcurrentBag<BlobIdentifier> missingBlobs = new ConcurrentBag<BlobIdentifier>();
+
+        try
+        {
+            await Parallel.ForEachAsync(blobs, async (identifier, ctx) =>
+            {
+                bool exists = await Exists(ns, identifier);
+
+                if (!exists)
+                {
+                    missingBlobs.Add(identifier);
+                }
+            });
+        }
+        catch (AggregateException e)
+        {
+            if (e.InnerException is PartialReferenceResolveException)
+            {
+                throw e.InnerException;
+            }
+
+            throw;
+        }
+
+        return missingBlobs.ToArray();
     }
 
     public async Task<BlobIdentifier[]> FilterOutKnownBlobs(NamespaceId ns, IAsyncEnumerable<BlobIdentifier> blobs)
