UnrealCloudDDC - Futher tweaks to make FilterOutKnownBlobs more parallel

diff --git a/Engine/Source/Programs/UnrealCloudDDC/Jupiter/Implementation/Blob/IBlobService.cs b/Engine/Source/Programs/UnrealCloudDDC/Jupiter/Implementation/Blob/IBlobService.cs
--- a/Engine/Source/Programs/UnrealCloudDDC/Jupiter/Implementation/Blob/IBlobService.cs
+++ b/Engine/Source/Programs/UnrealCloudDDC/Jupiter/Implementation/Blob/IBlobService.cs
@@ -569,28 +569,21 @@ public class BlobService : IBlobService
 
     public async Task<BlobIdentifier[]> FilterOutKnownBlobs(NamespaceId ns, IEnumerable<BlobIdentifier> blobs)
     {
-        ConcurrentBag<BlobIdentifier> missingBlobs = new ConcurrentBag<BlobIdentifier>();
-
-        try
+        List<(BlobIdentifier, Task<bool>)> existTasks = new();
+        foreach (BlobIdentifier blob in blobs)
         {
-            await Parallel.ForEachAsync(blobs, async (identifier, ctx) =>
-            {
-                bool exists = await Exists(ns, identifier);
-
-                if (!exists)
-                {
-                    missingBlobs.Add(identifier);
-                }
-            });
+            existTasks.Add((blob, Exists(ns, blob)));
         }
-        catch (AggregateException e)
+
+        List<BlobIdentifier> missingBlobs = new();
+        foreach ((BlobIdentifier blob, Task<bool> existsTask) in existTasks)
         {
-            if (e.InnerException is PartialReferenceResolveException)
+            bool exists = await existsTask;
+
+            if (!exists)
             {
-                throw e.InnerException;
+                missingBlobs.Add(blob);
             }
-
-            throw;
         }
 
         return missingBlobs.ToArray();
