UnrealCloudDDC - Content Id attachment resolve is now able to go wide on processing of each attachment
Also removed redundant exist check when finalizing objects.

diff --git a/Engine/Source/Programs/UnrealCloudDDC/Jupiter/Implementation/IReferenceResolver.cs b/Engine/Source/Programs/UnrealCloudDDC/Jupiter/Implementation/IReferenceResolver.cs
--- a/Engine/Source/Programs/UnrealCloudDDC/Jupiter/Implementation/IReferenceResolver.cs
+++ b/Engine/Source/Programs/UnrealCloudDDC/Jupiter/Implementation/IReferenceResolver.cs
@@ -66,8 +66,23 @@ namespace Jupiter.Implementation
 
     public interface IReferenceResolver
     {
-
+        /// <summary>
+        /// Returns the blobs referenced from the cb object and any children
+        /// These blobs are guaranteed to exist, if any blob is missing
+        /// PartialReferenceResolveException or ReferenceIsMissingBlobsException are raised
+        /// </summary>
+        /// <param name="ns">The namespace to check</param>
+        /// <param name="cb">The compact binary object to resolve references for</param>
+        /// <returns></returns>
         IAsyncEnumerable<BlobIdentifier> GetReferencedBlobs(NamespaceId ns, CbObject cb);
+
+        /// <summary>
+        /// Returns which attachments exist in the cb object or any children
+        /// These attachments are not guaranteed to reference blobs that actually exist
+        /// </summary>
+        /// <param name="ns">The namespace to check</param>
+        /// <param name="cb">The compact binary object to resolve references for</param>
+        /// <returns></returns>
         IAsyncEnumerable<Attachment> GetAttachments(NamespaceId ns, CbObject cb);
     }
 
@@ -213,6 +228,7 @@ namespace Jupiter.Implementation
         public async IAsyncEnumerable<BlobIdentifier> GetReferencedBlobs(NamespaceId ns, CbObject cb)
         {
             List<Task<(BlobIdentifier, bool)>> pendingBlobExistsChecks = new();
+            List<Task<(ContentIdAttachment, bool)>> pendingContentIdChecks = new();
             List<ContentId> unresolvedContentIdReferences = new List<ContentId>();
             List<BlobIdentifier> unresolvedBlobReferences = new List<BlobIdentifier>();
 
@@ -226,17 +242,7 @@ namespace Jupiter.Implementation
                 else if (attachment is ContentIdAttachment contentIdAttachment)
                 {
                     // If we find a content id we resolve that into the actual blobs it references
-                    foreach (BlobIdentifier b in contentIdAttachment.ReferencedBlobs)
-                    {
-                        (BlobIdentifier _, bool exists) = await CheckBlobExists(ns, b);
-                        if (!exists)
-                        {
-                            unresolvedContentIdReferences.Add(contentIdAttachment.Identifier);
-                            continue;
-                        }
-
-                        yield return b;
-                    }
+                    pendingContentIdChecks.Add(CheckContentIdExists(ns, contentIdAttachment));
                 }
                 else if (attachment is ObjectAttachment objectAttachment)
                 {
@@ -263,6 +269,22 @@ namespace Jupiter.Implementation
                 }
             }
 
+            foreach (Task<(ContentIdAttachment, bool)> pendingContentIdTask in pendingContentIdChecks)
+            {
+                (ContentIdAttachment contentIdAttachment, bool exists) = await pendingContentIdTask;
+
+                if (!exists)
+                {
+                    unresolvedContentIdReferences.Add(contentIdAttachment.Identifier);
+                    continue;
+                }
+
+                foreach (BlobIdentifier b in contentIdAttachment.ReferencedBlobs)
+                {
+                    yield return b;
+                }
+            }
+
             // if there were any content ids we did not recognize we throw a partial reference exception
             if (unresolvedContentIdReferences.Count != 0)
             {
@@ -275,6 +297,23 @@ namespace Jupiter.Implementation
             }
         }
 
+        private async Task<(ContentIdAttachment, bool)> CheckContentIdExists(NamespaceId ns, ContentIdAttachment contentIdAttachment)
+        {
+            bool allBlobsExist = true;
+            foreach (BlobIdentifier b in contentIdAttachment.ReferencedBlobs)
+            {
+                (BlobIdentifier _, bool exists) = await CheckBlobExists(ns, b);
+
+                if (!exists)
+                {
+                    allBlobsExist = false;
+                    break;
+                }
+            }
+
+            return (contentIdAttachment, allBlobsExist);
+        }
+
         private async Task<CbObject> ParseCompactBinaryAttachment(NamespaceId ns, BlobIdentifier blobIdentifier)
         {
             BlobContents contents = await _blobStore.GetObject(ns, blobIdentifier);
diff --git a/Engine/Source/Programs/UnrealCloudDDC/Jupiter/Implementation/ObjectService.cs b/Engine/Source/Programs/UnrealCloudDDC/Jupiter/Implementation/ObjectService.cs
--- a/Engine/Source/Programs/UnrealCloudDDC/Jupiter/Implementation/ObjectService.cs
+++ b/Engine/Source/Programs/UnrealCloudDDC/Jupiter/Implementation/ObjectService.cs
@@ -116,6 +116,9 @@ namespace Jupiter.Implementation
 
         public async Task<(ContentId[], BlobIdentifier[])> Put(NamespaceId ns, BucketId bucket, IoHashKey key, BlobIdentifier blobHash, CbObject payload)
         {
+            IServerTiming? serverTiming = _httpContextAccessor.HttpContext?.RequestServices.GetService<IServerTiming>();
+            using ServerTimingMetricScoped? serverTimingScope = serverTiming?.CreateServerTimingMetricScope("ref.put", "Inserting ref");
+
             bool hasReferences = HasAttachments(payload);
 
             // if we have no references we are always finalized, e.g. there are no referenced blobs to upload
@@ -183,6 +186,9 @@ namespace Jupiter.Implementation
         
         private async Task<(ContentId[], BlobIdentifier[])> DoFinalize(NamespaceId ns, BucketId bucket, IoHashKey key, BlobIdentifier blobHash, CbObject payload)
         {
+            IServerTiming? serverTiming = _httpContextAccessor.HttpContext?.RequestServices.GetService<IServerTiming>();
+            using ServerTimingMetricScoped? serverTimingScope = serverTiming?.CreateServerTimingMetricScope("ref.finalize", "Finalizing the ref");
+
             Task addRefToBlobsTask = _blobIndex.AddRefToBlobs(ns, bucket, key, new [] {blobHash});
 
             ContentId[] missingReferences = Array.Empty<ContentId>();
@@ -195,9 +201,8 @@ namespace Jupiter.Implementation
                 {
                     IAsyncEnumerable<BlobIdentifier> references = _referenceResolver.GetReferencedBlobs(ns, payload);
                     BlobIdentifier[] referencesArray = await references.ToArrayAsync();
-                    Task addRefsTask = _blobIndex.AddRefToBlobs(ns, bucket, key, referencesArray);
-                    missingBlobs = await _blobService.FilterOutKnownBlobs(ns, referencesArray);
-                    await addRefsTask;
+                    // TODO: Blobs could be added to the blob index as we find them in the async enumerable
+                    await _blobIndex.AddRefToBlobs(ns, bucket, key, referencesArray);
                 }
                 catch (PartialReferenceResolveException e)
                 {
