UnrealCloudDDC - Removing logging of state when starting a background polling service. This can potentially be quite large and we have reports of crashes from the logger under certain conditions. It does not really provide a lot of useful information.

diff --git a/Engine/Source/Programs/UnrealCloudDDC/Jupiter.Common/PollingService.cs b/Engine/Source/Programs/UnrealCloudDDC/Jupiter.Common/PollingService.cs
--- a/Engine/Source/Programs/UnrealCloudDDC/Jupiter.Common/PollingService.cs
+++ b/Engine/Source/Programs/UnrealCloudDDC/Jupiter.Common/PollingService.cs
@@ -62,7 +62,7 @@ namespace Jupiter
         public Task StartAsync(CancellationToken cancellationToken)
         {
             bool shouldPoll = ShouldStartPolling();
-            _logger.LogInformation("Polling service {Service} initialized {@State} , will poll: {WillPoll}.", _serviceName, _state, shouldPoll);
+            _logger.LogInformation("Polling service {Service} initialized, will poll: {WillPoll}.", _serviceName, shouldPoll);
 
             if (shouldPoll)
             {
