UnrealCloudDDC - Tweaks to the helm charts allowing for specification of TLS and ingress class name in the ingress, overriding the scopes used for the service credentials as well as overriding annotations and labels on the service account

diff --git a/Engine/Source/Programs/UnrealCloudDDC/Helm/UnrealCloudDDC/templates/ingress.yaml b/Engine/Source/Programs/UnrealCloudDDC/Helm/UnrealCloudDDC/templates/ingress.yaml
--- a/Engine/Source/Programs/UnrealCloudDDC/Helm/UnrealCloudDDC/templates/ingress.yaml
+++ b/Engine/Source/Programs/UnrealCloudDDC/Helm/UnrealCloudDDC/templates/ingress.yaml
@@ -9,9 +9,12 @@ metadata:
     {{ $key }}: {{ $value | quote }}
     {{- end }}
 spec:
+  {{- if and .Values.ingress.ingressClassName (include "common.ingress.supportsIngressClassname" .) }}
+  ingressClassName: {{ .Values.ingress.ingressClassName | quote }}
+  {{- end }}
   rules:
     {{- if .Values.ingress.hostname }}
-    - host: {{ .Values.ingress.hostname }}
+    - host: {{ .Values.ingress.hostname | quote }}
       http:
         paths:
           {{- if .Values.ingress.extraPaths }}
@@ -33,4 +36,15 @@ spec:
             {{- end }}
             backend: {{- include "common.ingress.backend" (dict "serviceName" (include "common.names.fullname" $) "servicePort" "http" "context" $) | nindent 14 }}
     {{- end }}
+  {{- if or .Values.ingress.tls .Values.ingress.extraTls .Values.ingress.hosts }}
+  tls:
+    {{- if .Values.ingress.tls }}
+    - hosts:
+        - {{ .Values.ingress.hostname }}
+      secretName: {{ printf "%s-tls" .Values.ingress.hostname | trunc 63 | trimSuffix "-" }}
+    {{- end }}
+    {{- if .Values.ingress.extraTls }}
+    {{- toYaml .Values.ingress.extraTls | nindent 4 }}
+    {{- end }}
+  {{- end }}
 {{- end }}
\ No newline at end of file
diff --git a/Engine/Source/Programs/UnrealCloudDDC/Helm/UnrealCloudDDC/templates/serviceaccount.yaml b/Engine/Source/Programs/UnrealCloudDDC/Helm/UnrealCloudDDC/templates/serviceaccount.yaml
--- a/Engine/Source/Programs/UnrealCloudDDC/Helm/UnrealCloudDDC/templates/serviceaccount.yaml
+++ b/Engine/Source/Programs/UnrealCloudDDC/Helm/UnrealCloudDDC/templates/serviceaccount.yaml
@@ -5,6 +5,16 @@ metadata:
   name: {{ template "jupiter.serviceAccountName" . }}
   namespace: {{ .Release.Namespace | quote }}
   labels: {{- include "common.labels.standard" . | nindent 4 }}
+    {{- if .Values.commonLabels }}
+    {{- include "common.tplvalues.render" ( dict "value" .Values.commonLabels "context" $ ) | nindent 4 }}
+    {{- end }}
+  annotations:
+    {{- if .Values.serviceAccount.annotations }}
+    {{- include "common.tplvalues.render" ( dict "value" .Values.serviceAccount.annotations "context" $) | nindent 4 }}
+    {{- end }}
+    {{- if .Values.commonAnnotations }}
+    {{- include "common.tplvalues.render" ( dict "value" .Values.commonAnnotations "context" $ ) | nindent 4 }}
+    {{- end }}
 secrets:
   - name: {{ include "common.names.fullname" . }}
 {{- end }}
