UnrealCloudDDC - Lower scylla timeout to 2 second as we generally expect requests to finish much faster then that.

diff --git a/Engine/Source/Programs/UnrealCloudDDC/Jupiter/Implementation/Objects/ScyllaSettings.cs b/Engine/Source/Programs/UnrealCloudDDC/Jupiter/Implementation/Objects/ScyllaSettings.cs
--- a/Engine/Source/Programs/UnrealCloudDDC/Jupiter/Implementation/Objects/ScyllaSettings.cs
+++ b/Engine/Source/Programs/UnrealCloudDDC/Jupiter/Implementation/Objects/ScyllaSettings.cs
@@ -58,6 +58,12 @@ public class ScyllaSettings : IValidatableObject
     /// </summary>
     public bool UseSSL { get; set; } = true;
 
+    /// <summary>
+    /// Read timeout in milliseconds
+    /// Set to -1 to get the default timeout, set to 0 to disable timeouts
+    /// </summary>
+    public int ReadTimeout { get; set; } = 2000; // timeout after 2 seconds
+
     /// <summary>
     /// Enable to scan database per shard, this requires setting more information about your cluster
     /// </summary>
diff --git a/Engine/Source/Programs/UnrealCloudDDC/Jupiter/Jupiter.Startup.cs b/Engine/Source/Programs/UnrealCloudDDC/Jupiter/Jupiter.Startup.cs
--- a/Engine/Source/Programs/UnrealCloudDDC/Jupiter/Jupiter.Startup.cs
+++ b/Engine/Source/Programs/UnrealCloudDDC/Jupiter/Jupiter.Startup.cs
@@ -20,7 +20,6 @@ using Jupiter.Implementation.Blob;
 using Jupiter.Implementation.LeaderElection;
 using Jupiter.Common.Implementation;
 using Jupiter.Implementation.TransactionLog;
-using Jupiter.Utils;
 using Microsoft.AspNetCore.Authorization;
 using Microsoft.Extensions.Configuration;
 using Microsoft.Extensions.DependencyInjection;
@@ -58,7 +57,7 @@ namespace Jupiter
         {
             services.AddHttpClient(Options.DefaultName, (provider, client) =>
             {
-                string GetPlatformName()
+                static string GetPlatformName()
                 {
                     if (RuntimeInformation.IsOSPlatform(OSPlatform.Linux))
                     {
@@ -296,6 +295,10 @@ namespace Jupiter
                 .WithExecutionProfiles(options =>
                     options.WithProfile("default", builder => builder.WithConsistencyLevel(ConsistencyLevel.LocalOne)));
 
+            if (settings.ReadTimeout != -1)
+            {
+                clusterBuilder.SocketOptions.SetReadTimeoutMillis(settings.ReadTimeout);
+            }
             if (settings.UseAzureCosmosDB)
             {
                 CassandraConnectionStringBuilder connectionStringBuilder = new CassandraConnectionStringBuilder(connectionString);
