UnrealCloudDDC - Upgraded bitnami common to 2.2.2 - adding support for determining if ingressClassName is supported

diff --git a/Engine/Source/Programs/UnrealCloudDDC/Helm/UnrealCloudDDC/Chart.yaml b/Engine/Source/Programs/UnrealCloudDDC/Helm/UnrealCloudDDC/Chart.yaml
--- a/Engine/Source/Programs/UnrealCloudDDC/Helm/UnrealCloudDDC/Chart.yaml
+++ b/Engine/Source/Programs/UnrealCloudDDC/Helm/UnrealCloudDDC/Chart.yaml
@@ -20,4 +20,4 @@ dependencies:
   repository: https://charts.bitnami.com/bitnami
   tags:
     - bitnami-common
-  version: 1.x.x
\ No newline at end of file
+  version: 2.x.x
\ No newline at end of file
