UnrealCloudDDC - Use the default scylla timeout by default as we have seen some issues with lowering this value. Will still override this but in a more selective way to see what are reasonable values to use.

diff --git a/Engine/Source/Programs/UnrealCloudDDC/Jupiter/Implementation/Objects/ScyllaSettings.cs b/Engine/Source/Programs/UnrealCloudDDC/Jupiter/Implementation/Objects/ScyllaSettings.cs
--- a/Engine/Source/Programs/UnrealCloudDDC/Jupiter/Implementation/Objects/ScyllaSettings.cs
+++ b/Engine/Source/Programs/UnrealCloudDDC/Jupiter/Implementation/Objects/ScyllaSettings.cs
@@ -62,7 +62,7 @@ public class ScyllaSettings : IValidatableObject
     /// Read timeout in milliseconds
     /// Set to -1 to get the default timeout, set to 0 to disable timeouts
     /// </summary>
-    public int ReadTimeout { get; set; } = 2000; // timeout after 2 seconds
+    public int ReadTimeout { get; set; } = -1; // use the default timeout value by default
 
     /// <summary>
     /// Enable to scan database per shard, this requires setting more information about your cluster
