Unreal Cloud DDC - Increasing time we are waiting for Scylla to startup before executing the tests as we are seeing a sharp increase in times when its not ready when start running tests

diff --git a/Engine/Source/Programs/UnrealCloudDDC/BuildUnrealCloudDDC.xml b/Engine/Source/Programs/UnrealCloudDDC/BuildUnrealCloudDDC.xml
--- a/Engine/Source/Programs/UnrealCloudDDC/BuildUnrealCloudDDC.xml
+++ b/Engine/Source/Programs/UnrealCloudDDC/BuildUnrealCloudDDC.xml
@@ -18,7 +18,7 @@
 			<Docker-Compose-Up File="$(ComposeFile)"/>
 
 			<!-- Sleep a while to let scylla finish starting up before we execute the tests -->
-			<Spawn Exe="/bin/bash" Arguments="-c 'sleep 60'" />
+			<Spawn Exe="/bin/bash" Arguments="-c 'sleep 100'" />
 
 			<!-- Run the unit tests, and override the dotnet sdk used to the global one for the machine rather then what is in the sync, to allow us to use newer versions then the rest of the engine has migrated to -->
 			<DotNet Arguments="test Engine/Source/Programs/UnrealCloudDDC/UnrealCloudDDC.sln -l console;verbosity=normal" DotNetPath="/usr/share/dotnet/dotnet"/>
