UnrealCloudDDC - If there are only pending content id resolves, wait for one of them to finish before attempting to look for finished results. Avoids a potential busy wait condition while resolving content ids.

diff --git a/Engine/Source/Programs/UnrealCloudDDC/Jupiter/Implementation/IReferenceResolver.cs b/Engine/Source/Programs/UnrealCloudDDC/Jupiter/Implementation/IReferenceResolver.cs
--- a/Engine/Source/Programs/UnrealCloudDDC/Jupiter/Implementation/IReferenceResolver.cs
+++ b/Engine/Source/Programs/UnrealCloudDDC/Jupiter/Implementation/IReferenceResolver.cs
@@ -207,12 +207,20 @@ namespace Jupiter.Implementation
                     pendingCompactBinaryAttachments.Remove(finishedTask);
                 }
 
+                bool hasWaited = false;
                 // if there are pending resolves left, wait for one of them to finish to avoid busy waiting
                 if (pendingCompactBinaryAttachments.Any())
                 {
+                    hasWaited = true;
                     await Task.WhenAny(pendingCompactBinaryAttachments);
                 }
 
+                // if we did not waiting for compact binary attachment resolves we consider waiting for content id resolves if there are any
+                if (!hasWaited && pendingContentIdResolves.Any())
+                {
+                    await Task.WhenAny(pendingContentIdResolves);
+                }
+
                 foreach (Attachment attachment in attachments)
                 {
                     yield return attachment;
