Unreal Cloud DDC - Removed or reduced level of some frequent logging that we do not really look at very often in a attempt to reduce spam.

diff --git a/Engine/Source/Programs/UnrealCloudDDC/Jupiter/Implementation/Replication/RefsReplicator.cs b/Engine/Source/Programs/UnrealCloudDDC/Jupiter/Implementation/Replication/RefsReplicator.cs
--- a/Engine/Source/Programs/UnrealCloudDDC/Jupiter/Implementation/Replication/RefsReplicator.cs
+++ b/Engine/Source/Programs/UnrealCloudDDC/Jupiter/Implementation/Replication/RefsReplicator.cs
@@ -613,12 +613,11 @@ namespace Jupiter.Implementation
 
         private void LogReplicationHeartbeat(int countOfCurrentReplications)
         {
-            
             // log message used to generate metric for how many replications are currently running
             _logger.LogInformation("{Name} replication has run . Count of running replications: {CurrentReplications}", _name, countOfCurrentReplications);
 
             // log message used to verify replicators are actually running
-            _logger.LogInformation("{Name} starting replication. Last transaction was {TransactionId} {Generation}", _name, State.ReplicatorOffset.GetValueOrDefault(0L), State.ReplicatingGeneration.GetValueOrDefault(Guid.Empty) );
+            _logger.LogDebug("{Name} starting replication. Last transaction was {TransactionId} {Generation}", _name, State.ReplicatorOffset.GetValueOrDefault(0L), State.ReplicatingGeneration.GetValueOrDefault(Guid.Empty) );
         }
 
         private async Task AddToReplicationLog(NamespaceId ns, BucketId bucket, IoHashKey key, BlobIdentifier blob)
