UnrealCloudDDC - Restored logging event used to track replicator state for our monitors.

diff --git a/Engine/Source/Programs/UnrealCloudDDC/Jupiter/Implementation/Replication/RefsReplicator.cs b/Engine/Source/Programs/UnrealCloudDDC/Jupiter/Implementation/Replication/RefsReplicator.cs
--- a/Engine/Source/Programs/UnrealCloudDDC/Jupiter/Implementation/Replication/RefsReplicator.cs
+++ b/Engine/Source/Programs/UnrealCloudDDC/Jupiter/Implementation/Replication/RefsReplicator.cs
@@ -617,7 +617,7 @@ namespace Jupiter.Implementation
             _logger.LogInformation("{Name} replication has run . Count of running replications: {CurrentReplications}", _name, countOfCurrentReplications);
 
             // log message used to verify replicators are actually running
-            _logger.LogDebug("{Name} starting replication. Last transaction was {TransactionId} {Generation}", _name, State.ReplicatorOffset.GetValueOrDefault(0L), State.ReplicatingGeneration.GetValueOrDefault(Guid.Empty) );
+            _logger.LogInformation("{Name} starting replication. Last transaction was {TransactionId} {Generation}", _name, State.ReplicatorOffset.GetValueOrDefault(0L), State.ReplicatingGeneration.GetValueOrDefault(Guid.Empty) );
         }
 
         private async Task AddToReplicationLog(NamespaceId ns, BucketId bucket, IoHashKey key, BlobIdentifier blob)
