UnrealCloudDDC - Fix for issue with content ids when using a Cassandra (not Scylla) database

diff --git a/Engine/Source/Programs/UnrealCloudDDC/Jupiter/Implementation/ContentId/ScyllaContentIdStore.cs b/Engine/Source/Programs/UnrealCloudDDC/Jupiter/Implementation/ContentId/ScyllaContentIdStore.cs
--- a/Engine/Source/Programs/UnrealCloudDDC/Jupiter/Implementation/ContentId/ScyllaContentIdStore.cs
+++ b/Engine/Source/Programs/UnrealCloudDDC/Jupiter/Implementation/ContentId/ScyllaContentIdStore.cs
@@ -85,7 +85,7 @@ namespace Jupiter.Implementation
                             throw new InvalidContentIdException(contentId);
                         }
 
-                        BlobIdentifier[] blobs = resolvedContentId.Chunks.Select(b => b.AsBlobIdentifier()).ToArray();
+                        BlobIdentifier[] blobs = resolvedContentId.Chunks.Select(b => new BlobIdentifier(b!)).ToArray();
 
                         {
                             using TelemetrySpan _ = _tracer.StartActiveSpan("ScyllaContentIdStore.FindMissingBlobs").SetAttribute("operation.name", "ScyllaContentIdStore.FindMissingBlobs");
@@ -120,7 +120,6 @@ namespace Jupiter.Implementation
         public async Task Put(NamespaceId ns, ContentId contentId, BlobIdentifier blobIdentifier, int contentWeight)
         {
             using TelemetrySpan scope = _tracer.BuildScyllaSpan("ScyllaContentIdStore.PutContentId");
-
             if (_scyllaSessionManager.IsScylla)
             {
                 await _mapper.UpdateAsync<ScyllaContentId>("SET chunks = ? WHERE content_id = ? AND content_weight = ?", new [] {new ScyllaBlobIdentifier(blobIdentifier)}, new ScyllaBlobIdentifier(contentId), contentWeight);
@@ -128,7 +127,7 @@ namespace Jupiter.Implementation
             }
             else
             {
-                await _mapper.UpdateAsync<CassandraContentId>("SET chunks = ? WHERE content_id = ? AND content_weight = ?", new [] {new ScyllaBlobIdentifier(blobIdentifier)}, contentId.HashData, contentWeight);
+                await _mapper.UpdateAsync<CassandraContentId>("SET chunks = ? WHERE content_id = ? AND content_weight = ?", new [] {blobIdentifier.HashData}, contentId.HashData, contentWeight);
             }
         }
     }
@@ -172,7 +171,7 @@ namespace Jupiter.Implementation
         {
             ContentId = contentId.HashData;
             ContentWeight = contentWeight;
-            Chunks = chunks.Select(b => new ScyllaBlobIdentifier(b)).ToArray();
+            Chunks = chunks.Select(b => b.HashData).ToArray();
         }
 
         [Cassandra.Mapping.Attributes.PartitionKey]
@@ -183,6 +182,6 @@ namespace Jupiter.Implementation
         [Cassandra.Mapping.Attributes.Column("content_weight")]
         public int? ContentWeight { get; set; }
 
-        public ScyllaBlobIdentifier[] Chunks { get; set; } = Array.Empty<ScyllaBlobIdentifier>();
+        public byte[][] Chunks { get; set; } = Array.Empty<byte[]>();
     }
 }
