UnrealCloudDDC - Updated changelog

diff --git a/Engine/Source/Programs/UnrealCloudDDC/changelog.md b/Engine/Source/Programs/UnrealCloudDDC/changelog.md
--- a/Engine/Source/Programs/UnrealCloudDDC/changelog.md
+++ b/Engine/Source/Programs/UnrealCloudDDC/changelog.md
@@ -3,6 +3,9 @@
 * Added deletion of invalid references during blob GC - this reduces the size for entries in the blob_index table.
 * Using OpenTelemetry for tracing - still supports datadog traces but this can now be forwarded to any OpenTelemetry compatible service.
 * Fixed issue in Azure Blob storage with namespace containing _
+* Optimized content id resolving
+* Fixed issue with crash on startup from background services for certain large configuration objects
+* Fixed issues with content id when using CosmosDB
 
 # 0.1.0
 * First public release of UnrealCloudDDC
