Prepared release 0.2.0 of unreal cloud ddc

diff --git a/Engine/Source/Programs/UnrealCloudDDC/Helm/UnrealCloudDDC/Chart.yaml b/Engine/Source/Programs/UnrealCloudDDC/Helm/UnrealCloudDDC/Chart.yaml
--- a/Engine/Source/Programs/UnrealCloudDDC/Helm/UnrealCloudDDC/Chart.yaml
+++ b/Engine/Source/Programs/UnrealCloudDDC/Helm/UnrealCloudDDC/Chart.yaml
@@ -4,7 +4,7 @@ description: A Helm chart for Kubernetes deploying UnrealCloudDDC
 
 type: application
 
-version: 0.1.0
+version: 0.2.0
 
 icon: https://devtools-misc.s3.amazonaws.com/ue4logo.png
 home: https://unrealengine.com
diff --git a/Engine/Source/Programs/UnrealCloudDDC/Helm/UnrealCloudDDC/version.yaml b/Engine/Source/Programs/UnrealCloudDDC/Helm/UnrealCloudDDC/version.yaml
--- a/Engine/Source/Programs/UnrealCloudDDC/Helm/UnrealCloudDDC/version.yaml
+++ b/Engine/Source/Programs/UnrealCloudDDC/Helm/UnrealCloudDDC/version.yaml
@@ -1 +1 @@
-version: 0.1.0
\ No newline at end of file
+version: 0.2.0
\ No newline at end of file
diff --git a/Engine/Source/Programs/UnrealCloudDDC/changelog.md b/Engine/Source/Programs/UnrealCloudDDC/changelog.md
--- a/Engine/Source/Programs/UnrealCloudDDC/changelog.md
+++ b/Engine/Source/Programs/UnrealCloudDDC/changelog.md
@@ -1,4 +1,6 @@
 # Unreleased
+
+# 0.2.0
 * Improved handling of cassandra timeouts during ref GC.
 * Added deletion of invalid references during blob GC - this reduces the size for entries in the blob_index table.
 * Using OpenTelemetry for tracing - still supports datadog traces but this can now be forwarded to any OpenTelemetry compatible service.
diff --git a/Engine/Source/Programs/UnrealCloudDDC/version.yaml b/Engine/Source/Programs/UnrealCloudDDC/version.yaml
--- a/Engine/Source/Programs/UnrealCloudDDC/version.yaml
+++ b/Engine/Source/Programs/UnrealCloudDDC/version.yaml
@@ -1 +1 @@
-version: 0.1.0
\ No newline at end of file
+version: 0.2.0
\ No newline at end of file
