UnrealCloudDDC - Added support for initContainers in helm

diff --git a/Engine/Source/Programs/UnrealCloudDDC/Helm/UnrealCloudDDC/templates/deployment.yaml b/Engine/Source/Programs/UnrealCloudDDC/Helm/UnrealCloudDDC/templates/deployment.yaml
--- a/Engine/Source/Programs/UnrealCloudDDC/Helm/UnrealCloudDDC/templates/deployment.yaml
+++ b/Engine/Source/Programs/UnrealCloudDDC/Helm/UnrealCloudDDC/templates/deployment.yaml
@@ -41,6 +41,9 @@ spec:
       {{- if .Values.tolerations }}
       tolerations: {{- include "common.tplvalues.render" (dict "value" .Values.tolerations "context" $) | nindent 8 }}
       {{- end }}
+      {{- if .Values.initContainers }}
+      initContainers: {{- include "common.tplvalues.render" (dict "value" .Values.initContainers "context" $) | nindent 8 }}
+      {{- end }}
       containers:
       - name: {{ include "common.names.name" . }}
         image: {{ printf "%s:%s" .Values.image.repository (include "jupiter.appversion" . ) }}
