Unreal Cloud DDC - Default filesystem max size from persistence if configured

diff --git a/Engine/Source/Programs/UnrealCloudDDC/Helm/UnrealCloudDDC/templates/configmap.yaml b/Engine/Source/Programs/UnrealCloudDDC/Helm/UnrealCloudDDC/templates/configmap.yaml
--- a/Engine/Source/Programs/UnrealCloudDDC/Helm/UnrealCloudDDC/templates/configmap.yaml
+++ b/Engine/Source/Programs/UnrealCloudDDC/Helm/UnrealCloudDDC/templates/configmap.yaml
@@ -90,7 +90,7 @@ GC:
 {{- if or (.Values.persistence.enabled) (.Values.persistence.volume) }}
 FileSystem:
   RootDir:  {{ printf "%s/Blobs" .Values.persistence.mountPath }}
-  MaxSizeBytes: {{ required "MaxSizeBytes of the cache volume must be specified" $filesystem.MaxSizeBytes }}
+  MaxSizeBytes: {{ default .Values.persistence.size $filesystem.MaxSizeBytes }}
  {{- if $filesystem.TriggerThresholdPercentage }}
   TriggerThresholdPercentage: {{ $filesystem.TriggerThresholdPercentage }}
 {{- end }}
diff --git a/Engine/Source/Programs/UnrealCloudDDC/Helm/UnrealCloudDDC/values.yaml b/Engine/Source/Programs/UnrealCloudDDC/Helm/UnrealCloudDDC/values.yaml
--- a/Engine/Source/Programs/UnrealCloudDDC/Helm/UnrealCloudDDC/values.yaml
+++ b/Engine/Source/Programs/UnrealCloudDDC/Helm/UnrealCloudDDC/values.yaml
@@ -155,7 +155,7 @@ persistence:
 
   ## Persistent Volume size
   ##
-  size: 5Gi
+  size: 5242880 # 5 Gb specified in bytes so it can be used as a setting MaxSizeBytes
 
   ## Persistent Volume Storage Class
   ## If defined, storageClassName: <storageClass>
