Unreal Cloud DDC - Clarified how to set some required values to the helm chart

diff --git a/Engine/Source/Programs/UnrealCloudDDC/Helm/UnrealCloudDDC/templates/configmap.yaml b/Engine/Source/Programs/UnrealCloudDDC/Helm/UnrealCloudDDC/templates/configmap.yaml
--- a/Engine/Source/Programs/UnrealCloudDDC/Helm/UnrealCloudDDC/templates/configmap.yaml
+++ b/Engine/Source/Programs/UnrealCloudDDC/Helm/UnrealCloudDDC/templates/configmap.yaml
@@ -131,7 +131,7 @@ Azure:
 {{- $s3 := default dict $config.S3 -}}
 S3:
   ConnectionString: "AWS"
-  BucketName: {{ required "The name of the S3 bucket to use has to be specified" $s3.BucketName }}
+  BucketName: {{ required "The name of the S3 bucket to use has to be specified (.config.S3.BucketName)" $s3.BucketName }}
   # We assume that you are creating your buckets yourself for production deployments like this 
   CreateBucketIfMissing: {{ default  false $s3.CreateBucketIfMissing }}
   SetBucketPolicies: {{ default false $s3.SetBucketPolicies }}
