Unreal Cloud DDC - Added open telemetry metrics for GC data as a test to replace some of our excessive logging and events

diff --git a/Engine/Source/Programs/UnrealCloudDDC/Jupiter.Common/BaseStartup.cs b/Engine/Source/Programs/UnrealCloudDDC/Jupiter.Common/BaseStartup.cs
--- a/Engine/Source/Programs/UnrealCloudDDC/Jupiter.Common/BaseStartup.cs
+++ b/Engine/Source/Programs/UnrealCloudDDC/Jupiter.Common/BaseStartup.cs
@@ -3,6 +3,7 @@
 using System;
 using System.Collections.Generic;
 using System.ComponentModel.DataAnnotations;
+using System.Diagnostics.Metrics;
 using System.IdentityModel.Tokens.Jwt;
 using System.IO;
 using System.Net.Mime;
@@ -33,6 +34,7 @@ using Newtonsoft.Json.Serialization;
 using Okta.AspNet.Abstractions;
 using Okta.AspNetCore;
 using OpenTelemetry.Logs;
+using OpenTelemetry.Metrics;
 using OpenTelemetry.Resources;
 using OpenTelemetry.Trace;
 using Serilog;
@@ -257,6 +259,11 @@ namespace Jupiter
 
             string otelServiceName = Configuration["OTEL_SERVICE_NAME"] ?? "unreal-cloud-ddc";
             string otelServiceVersion = Configuration["OTEL_SERVICE_VERSION"];
+
+            ResourceBuilder appResourceBuilder = ResourceBuilder.CreateDefault()
+                .AddService("UnrealCloudDDC", serviceNamespace: "Jupiter", serviceVersion: otelServiceVersion)
+                .AddEnvironmentVariableDetector();
+
             services.AddOpenTelemetryTracing(builder =>
             {
                 builder.AddHttpClientInstrumentation(options =>
@@ -274,14 +281,9 @@ namespace Jupiter
                 });
                 builder.AddAspNetCoreInstrumentation();
 
+                builder.SetResourceBuilder(appResourceBuilder);
                 builder.AddOtlpExporter();
 
-                builder.ConfigureResource(resourceBuilder =>
-                {
-                    resourceBuilder.AddService("UnrealCloudDDC", serviceNamespace: "Jupiter", serviceVersion: otelServiceVersion);
-                    resourceBuilder.AddEnvironmentVariableDetector();
-                });
-
                 builder.AddSource("UnrealCloudDDC", "ScyllaDB");
             });
             services.Configure<OpenTelemetryLoggerOptions>(opt =>
@@ -293,6 +295,19 @@ namespace Jupiter
 
             services.AddSingleton<Tracer>(CreateTracer);
 
+            services.AddSingleton<Meter>(CreateMeter);
+
+            services.AddOpenTelemetryMetrics(metricProviderBuilder =>
+            {
+                metricProviderBuilder
+                    .AddMeter("UnrealCloudDDC", "ScyllaDB")
+                    .AddOtlpExporter()
+                    .AddConsoleExporter()
+                    .SetResourceBuilder(appResourceBuilder)
+                    .AddAspNetCoreInstrumentation()
+                    .AddHttpClientInstrumentation();
+            });
+
             services.Configure<ForwardedHeadersOptions>(options =>
             {
                 options.ForwardedHeaders = ForwardedHeaders.XForwardedFor | ForwardedHeaders.XForwardedProto;
@@ -325,6 +340,11 @@ namespace Jupiter
             OnAddHealthChecks(services);
         }
 
+        private Meter CreateMeter(IServiceProvider provider)
+        {
+            return new Meter("UnrealCloudDDC");
+        }
+
         private Tracer CreateTracer(IServiceProvider provider)
         {
             Tracer tracer = TracerProvider.Default.GetTracer("UnrealCloudDDC");
diff --git a/Engine/Source/Programs/UnrealCloudDDC/Jupiter.Common/Jupiter.Common.csproj b/Engine/Source/Programs/UnrealCloudDDC/Jupiter.Common/Jupiter.Common.csproj
--- a/Engine/Source/Programs/UnrealCloudDDC/Jupiter.Common/Jupiter.Common.csproj
+++ b/Engine/Source/Programs/UnrealCloudDDC/Jupiter.Common/Jupiter.Common.csproj
@@ -20,6 +20,7 @@
     <PackageReference Include="Azure.Identity" Version="1.5.0" />
     <PackageReference Include="Azure.Security.KeyVault.Secrets" Version="4.2.0" />
     <PackageReference Include="OpenTelemetry" Version="1.4.0-beta.3" />
+    <PackageReference Include="OpenTelemetry.Exporter.Console" Version="1.3.2" />
     <PackageReference Include="OpenTelemetry.Exporter.OpenTelemetryProtocol" Version="1.4.0-beta.3" />
     <PackageReference Include="OpenTelemetry.Instrumentation.AspNetCore" Version="1.0.0-rc9.9" />
     <PackageReference Include="OpenTelemetry.Instrumentation.Http" Version="1.0.0-rc9.9" />
diff --git a/Engine/Source/Programs/UnrealCloudDDC/Jupiter/Implementation/GC/OrphanBlobCleanupRefs.cs b/Engine/Source/Programs/UnrealCloudDDC/Jupiter/Implementation/GC/OrphanBlobCleanupRefs.cs
--- a/Engine/Source/Programs/UnrealCloudDDC/Jupiter/Implementation/GC/OrphanBlobCleanupRefs.cs
+++ b/Engine/Source/Programs/UnrealCloudDDC/Jupiter/Implementation/GC/OrphanBlobCleanupRefs.cs
@@ -11,6 +11,7 @@ using Jupiter.Common;
 using Microsoft.Extensions.Options;
 using OpenTelemetry.Trace;
 using Microsoft.Extensions.Logging;
+using System.Diagnostics.Metrics;
 
 namespace Jupiter.Implementation
 {
@@ -25,8 +26,11 @@ namespace Jupiter.Implementation
         private readonly Tracer _tracer;
         private readonly ILogger _logger;
 
+        private readonly Counter<long> _consideredBlobCounter;
+        private readonly Counter<long> _deletedBlobCounter;
+
         // ReSharper disable once UnusedMember.Global
-        public OrphanBlobCleanupRefs(IOptionsMonitor<GCSettings> gcSettings, IBlobService blobService, IObjectService objectService, IBlobIndex blobIndex, ILeaderElection leaderElection, INamespacePolicyResolver namespacePolicyResolver, Tracer tracer, ILogger<OrphanBlobCleanupRefs> logger)
+        public OrphanBlobCleanupRefs(IOptionsMonitor<GCSettings> gcSettings, IBlobService blobService, IObjectService objectService, IBlobIndex blobIndex, ILeaderElection leaderElection, INamespacePolicyResolver namespacePolicyResolver, Tracer tracer, ILogger<OrphanBlobCleanupRefs> logger, Meter meter)
         {
             _gcSettings = gcSettings;
             _blobService = blobService;
@@ -36,6 +40,8 @@ namespace Jupiter.Implementation
             _namespacePolicyResolver = namespacePolicyResolver;
             _tracer = tracer;
             _logger = logger;
+            _consideredBlobCounter = meter.CreateCounter<long>("gc.blob.considered");
+            _deletedBlobCounter = meter.CreateCounter<long>("gc.blob.deleted");
         }
 
         public bool ShouldRun()
@@ -109,7 +115,7 @@ namespace Jupiter.Implementation
                 TimeSpan storagePoolGcDuration = DateTime.Now - startTime;
                 _logger.LogInformation("Finished running Orphan GC For StoragePool: {StoragePool}. Took {Duration}", policy.StoragePool, storagePoolGcDuration);
             }
-
+            
             _logger.LogInformation("Finished running Orphan GC");
             return countOfBlobsRemoved;
         }
@@ -175,24 +181,28 @@ namespace Jupiter.Implementation
 
             removeBlobScope.SetAttribute("removed", (!found).ToString());
 
+            _consideredBlobCounter.Add(1);
+
             // something is still referencing this blob, we should not delete it
             if (found)
             {
                 return false;
             }
 
+            _deletedBlobCounter.Add(1);
+            _logger.LogInformation("GC Orphan blob {Blob} from {StoragePool} which was last modified at {LastModifiedTime}", blob, storagePoolName, lastModifiedTime);
+
             // if the blob was not found to have a reference in any of the namespace that share a storage pool then the blob is not used anymore and should be deleted from all the namespaces
             await Parallel.ForEachAsync(namespacesThatSharePool, cancellationToken, async (ns, _) =>
             {
-                await RemoveBlob(ns, blob, lastModifiedTime);
+                await RemoveBlob(ns, blob);
             });
             return true;
 
         }
 
-        private async Task RemoveBlob(NamespaceId ns, BlobIdentifier blob, DateTime lastModifiedTime)
+        private async Task RemoveBlob(NamespaceId ns, BlobIdentifier blob)
         {
-            _logger.LogInformation("Attempting to GC Orphan blob {Blob} from {Namespace} which was last modified at {LastModifiedTime}", blob, ns, lastModifiedTime);
             try
             {
                 await _blobService.DeleteObject(ns, blob);
