Unreal Cloud DDC - Missing file for reference cache that is causing build errors

diff --git a/Engine/Source/Programs/UnrealCloudDDC/Jupiter/Implementation/Objects/MemoryCachedReferencesStore.cs b/Engine/Source/Programs/UnrealCloudDDC/Jupiter/Implementation/Objects/MemoryCachedReferencesStore.cs
new file mode 100644
--- /dev/null
+++ b/Engine/Source/Programs/UnrealCloudDDC/Jupiter/Implementation/Objects/MemoryCachedReferencesStore.cs
@@ -0,0 +1,208 @@
+// Copyright Epic Games, Inc. All Rights Reserved.
+
+using System;
+using System.Collections.Concurrent;
+using System.Collections.Generic;
+using System.Threading.Tasks;
+using EpicGames.Horde.Storage;
+using Microsoft.Extensions.Caching.Memory;
+using Microsoft.Extensions.Options;
+
+namespace Jupiter.Implementation.Objects
+{
+    public class MemoryCachedReferencesStore : IReferencesStore
+    {
+        private readonly IReferencesStore _actualStore;
+        private readonly ConcurrentDictionary<NamespaceId, MemoryCache> _referenceCaches = new ConcurrentDictionary<NamespaceId, MemoryCache>();
+        private readonly IOptionsMonitor<MemoryCacheReferencesSettings> _options;
+
+        public MemoryCachedReferencesStore(IReferencesStore actualStore, IOptionsMonitor<MemoryCacheReferencesSettings> options)
+        {
+            _actualStore = actualStore;
+            _options = options;
+        }
+
+        private void AddCacheEntry(NamespaceId ns, BucketId bucket, IoHashKey key, ObjectRecord record)
+        {
+            // we can not cache none finalized records as they will be mutated again when finalized
+            if (!record.IsFinalized)
+            {
+                return;
+            }
+
+            MemoryCache cache = GetCacheForNamespace(ns);
+
+            CachedReferenceEntry cachedEntry = new CachedReferenceEntry(record);
+
+            using ICacheEntry entry = cache.CreateEntry(new CachedReferenceKey(bucket, key));
+            entry.Value = cachedEntry;
+            entry.Size = cachedEntry.Size;
+
+            if (_options.CurrentValue.EnableSlidingExpiry)
+            {
+                entry.SlidingExpiration = TimeSpan.FromMinutes(_options.CurrentValue.SlidingExpirationMinutes);
+            }
+        }
+
+        private MemoryCache GetCacheForNamespace(NamespaceId ns)
+        {
+            return _referenceCaches.GetOrAdd(ns, id => new MemoryCache(_options.CurrentValue));
+        }
+
+        public async Task<ObjectRecord> Get(NamespaceId ns, BucketId bucket, IoHashKey key, IReferencesStore.FieldFlags flags)
+        {
+            MemoryCache cache = GetCacheForNamespace(ns);
+
+            if (cache.TryGetValue(new CachedReferenceKey(bucket, key), out CachedReferenceEntry cachedResult))
+            {
+                return cachedResult.ToObjectRecord(flags);
+            }
+
+            ObjectRecord objectRecord = await _actualStore.Get(ns, bucket, key, IReferencesStore.FieldFlags.All);
+            AddCacheEntry(ns, bucket, key, objectRecord);
+
+            return objectRecord;
+        }
+
+        public Task Put(NamespaceId ns, BucketId bucket, IoHashKey key, BlobIdentifier blobHash, byte[] blob, bool isFinalized)
+        {
+            ObjectRecord objectRecord = new ObjectRecord(ns, bucket, key, DateTime.Now, blob, blobHash, isFinalized);
+            AddCacheEntry(ns, bucket, key, objectRecord);
+            
+            return _actualStore.Put(ns, bucket, key, blobHash, blob, isFinalized);
+        }
+
+        public Task Finalize(NamespaceId ns, BucketId bucket, IoHashKey key, BlobIdentifier blobIdentifier)
+        {
+            return _actualStore.Finalize(ns, bucket, key, blobIdentifier);
+        }
+
+        public Task UpdateLastAccessTime(NamespaceId ns, BucketId bucket, IoHashKey key, DateTime newLastAccessTime)
+        {
+            return _actualStore.UpdateLastAccessTime(ns, bucket, key, newLastAccessTime);
+        }
+
+        public IAsyncEnumerable<(BucketId, IoHashKey, DateTime)> GetRecords(NamespaceId ns)
+        {
+            return _actualStore.GetRecords(ns);
+        }
+
+        public IAsyncEnumerable<NamespaceId> GetNamespaces()
+        {
+            return _actualStore.GetNamespaces();
+        }
+
+        public Task<bool> Delete(NamespaceId ns, BucketId bucket, IoHashKey key)
+        {
+            MemoryCache cache = GetCacheForNamespace(ns);
+            cache.Remove(new CachedReferenceKey(bucket, key));
+            return _actualStore.Delete(ns, bucket, key);
+        }
+
+        public Task<long> DropNamespace(NamespaceId ns)
+        {
+            _referenceCaches.TryRemove(ns, out _);
+            return _actualStore.DropNamespace(ns);
+        }
+
+        public Task<long> DeleteBucket(NamespaceId ns, BucketId bucket)
+        {
+            // we do not track enough information to be able to drop a bucket, so we have to drop the entire namespace cache to remove the bucket
+            // this should be okay as deleting buckets is a extremely uncommon operation
+            _referenceCaches.TryRemove(ns, out _);
+            return _actualStore.DeleteBucket(ns, bucket);
+        }
+
+        public void Clear()
+        {
+            _referenceCaches.Clear();
+        }
+
+        public IReferencesStore GetUnderlyingStore()
+        {
+            return _actualStore;
+        }
+    }
+
+    class CachedReferenceKey : IEquatable<CachedReferenceKey>
+    {
+        private readonly BucketId _bucket;
+        private readonly IoHashKey _key;
+
+        public CachedReferenceKey(BucketId bucket, IoHashKey key)
+        {
+            _bucket = bucket;
+            _key = key;
+        }
+
+        public bool Equals(CachedReferenceKey? other)
+        {
+            if (ReferenceEquals(null, other))
+            {
+                return false;
+            }
+
+            if (ReferenceEquals(this, other))
+            {
+                return true;
+            }
+
+            return _bucket.Equals(other._bucket) && _key.Equals(other._key);
+        }
+
+        public override bool Equals(object? obj)
+        {
+            if (ReferenceEquals(null, obj))
+            {
+                return false;
+            }
+
+            if (ReferenceEquals(this, obj))
+            {
+                return true;
+            }
+
+            if (obj.GetType() != GetType())
+            {
+                return false;
+            }
+
+            return Equals((CachedReferenceKey)obj);
+        }
+
+        public override int GetHashCode()
+        {
+            return HashCode.Combine(_bucket, _key);
+        }
+    }
+
+    class CachedReferenceEntry
+    {
+        private int GetSize()
+        {
+            return Namespace.Text.Text.Length + Bucket.ToString().Length + 20 + Blob?.Length ?? 0 + 20;
+        }
+
+        public CachedReferenceEntry(ObjectRecord record)
+        {
+            Namespace = record.Namespace;
+            Bucket = record.Bucket;
+            Name = record.Name;
+            Blob = record.InlinePayload;
+            BlobIdentifier = record.BlobIdentifier;
+            Size = GetSize();
+        }
+
+        public NamespaceId Namespace { get; }
+        public BucketId Bucket { get; }
+        public IoHashKey Name { get;}
+        public byte[]? Blob { get; }
+        public BlobIdentifier BlobIdentifier { get; }
+        public int Size { get; }
+
+        public ObjectRecord ToObjectRecord(IReferencesStore.FieldFlags fieldFlags)
+        {
+            return new ObjectRecord(Namespace, Bucket, Name, DateTime.Now, (fieldFlags & IReferencesStore.FieldFlags.IncludePayload) != 0 ? Blob : null, BlobIdentifier, true);
+        }
+    }
+}
