Unreal Cloud DDC - Remove console exporter from metrics

diff --git a/Engine/Source/Programs/UnrealCloudDDC/Jupiter.Common/BaseStartup.cs b/Engine/Source/Programs/UnrealCloudDDC/Jupiter.Common/BaseStartup.cs
--- a/Engine/Source/Programs/UnrealCloudDDC/Jupiter.Common/BaseStartup.cs
+++ b/Engine/Source/Programs/UnrealCloudDDC/Jupiter.Common/BaseStartup.cs
@@ -302,7 +302,6 @@ namespace Jupiter
                 metricProviderBuilder
                     .AddMeter("UnrealCloudDDC", "ScyllaDB")
                     .AddOtlpExporter()
-                    .AddConsoleExporter()
                     .SetResourceBuilder(appResourceBuilder)
                     .AddAspNetCoreInstrumentation()
                     .AddHttpClientInstrumentation();
