Unreal Cloud DDC - Increase attempts at waiting for scylla to ready on startup, hopefully limits issue with tests failing due to scylla not being up.

diff --git a/Engine/Source/Programs/UnrealCloudDDC/Jupiter/Jupiter.Startup.cs b/Engine/Source/Programs/UnrealCloudDDC/Jupiter/Jupiter.Startup.cs
--- a/Engine/Source/Programs/UnrealCloudDDC/Jupiter/Jupiter.Startup.cs
+++ b/Engine/Source/Programs/UnrealCloudDDC/Jupiter/Jupiter.Startup.cs
@@ -356,7 +356,7 @@ namespace Jupiter
             }
 
             ISession replicatedSession;
-            const int MaxRetryAttempts = 10;
+            const int MaxRetryAttempts = 30;
             int countOfAttempts = 0;
             while(true)
             {
