Unreal Cloud DDC - Further increasing time we spend waiting for waiting for scylla to connect for tests

diff --git a/Engine/Source/Programs/UnrealCloudDDC/Jupiter/Jupiter.Startup.cs b/Engine/Source/Programs/UnrealCloudDDC/Jupiter/Jupiter.Startup.cs
--- a/Engine/Source/Programs/UnrealCloudDDC/Jupiter/Jupiter.Startup.cs
+++ b/Engine/Source/Programs/UnrealCloudDDC/Jupiter/Jupiter.Startup.cs
@@ -356,7 +356,7 @@ namespace Jupiter
             }
 
             ISession replicatedSession;
-            const int MaxRetryAttempts = 30;
+            const int MaxRetryAttempts = 50;
             int countOfAttempts = 0;
             while(true)
             {
