Unreal Cloud DDC - GC method

Added option to control how GC works per namespace for refs. The default (and only option previously) is based on last access. Can now be set to none (no GC) and TTL (Time to live), default time to live is set via the namespace policy.

Setting any gc mode except for last access will disable the last access tracking, but it can still be overriden to be disabled for the worker.
None GC will be used for Virtual Asset storage pools while TTL will be used for op logs and bundles.

diff --git a/Engine/Source/Programs/UnrealCloudDDC/Jupiter.Common/BaseStartup.cs b/Engine/Source/Programs/UnrealCloudDDC/Jupiter.Common/BaseStartup.cs
--- a/Engine/Source/Programs/UnrealCloudDDC/Jupiter.Common/BaseStartup.cs
+++ b/Engine/Source/Programs/UnrealCloudDDC/Jupiter.Common/BaseStartup.cs
@@ -724,5 +724,11 @@ namespace Jupiter
         public bool UseBlobIndexForExists { get; set; } = false;
         public bool UseBlobIndexForSlowExists { get; set; } = false;
         public bool IsPublicNamespace { get; set; } = true;
+
+        public enum StoragePoolGCMethod  { None, LastAccess, TTL };
+
+        public StoragePoolGCMethod GcMethod { get; set; } = StoragePoolGCMethod.LastAccess;
+
+        public TimeSpan DefaultTTL { get; set; } = TimeSpan.FromDays(14);
     }
 }
diff --git a/Engine/Source/Programs/UnrealCloudDDC/Jupiter/Implementation/GC/CacheCleanup.cs b/Engine/Source/Programs/UnrealCloudDDC/Jupiter/Implementation/GC/CacheCleanup.cs
--- a/Engine/Source/Programs/UnrealCloudDDC/Jupiter/Implementation/GC/CacheCleanup.cs
+++ b/Engine/Source/Programs/UnrealCloudDDC/Jupiter/Implementation/GC/CacheCleanup.cs
@@ -18,7 +18,7 @@ namespace Jupiter.Implementation
         Task<int> Cleanup(NamespaceId ns, CancellationToken cancellationToken);
     }
 
-    public class RefCleanup : IRefCleanup
+    public class RefLastAccessCleanup : IRefCleanup
     {
         private readonly IOptionsMonitor<GCSettings> _settings;
         private readonly IReferencesStore _referencesStore;
@@ -29,8 +29,8 @@ namespace Jupiter.Implementation
         private readonly IObjectService _objectService;
         private readonly OrphanBlobCleanupRefs _blobCleanup;
 
-        public RefCleanup(IOptionsMonitor<GCSettings> settings, IReferencesStore referencesStore,
-            IReplicationLog replicationLog, INamespacePolicyResolver namespacePolicyResolver, IObjectService objectService,  OrphanBlobCleanupRefs blobCleanup, Tracer tracer, ILogger<RefCleanup> logger)
+        public RefLastAccessCleanup(IOptionsMonitor<GCSettings> settings, IReferencesStore referencesStore,
+            IReplicationLog replicationLog, INamespacePolicyResolver namespacePolicyResolver, IObjectService objectService,  OrphanBlobCleanupRefs blobCleanup, Tracer tracer, ILogger<RefLastAccessCleanup> logger)
         {
             _settings = settings;
             _referencesStore = referencesStore;
@@ -62,6 +62,11 @@ namespace Jupiter.Implementation
                 return Task.FromResult(0);
             }
 
+            if (_namespacePolicyResolver.GetPoliciesForNs(ns).GcMethod != NamespacePolicy.StoragePoolGCMethod.LastAccess)
+            {
+                // only run for namepsaces set to use last access tracking
+            }
+
             return CleanNamespace(ns, cancellationToken);
         }
 
diff --git a/Engine/Source/Programs/UnrealCloudDDC/Jupiter/Implementation/LastAccess/LastAccessServiceReferences.cs b/Engine/Source/Programs/UnrealCloudDDC/Jupiter/Implementation/LastAccess/LastAccessServiceReferences.cs
--- a/Engine/Source/Programs/UnrealCloudDDC/Jupiter/Implementation/LastAccess/LastAccessServiceReferences.cs
+++ b/Engine/Source/Programs/UnrealCloudDDC/Jupiter/Implementation/LastAccess/LastAccessServiceReferences.cs
@@ -96,7 +96,12 @@ namespace Jupiter.Implementation
 
         private bool ShouldTrackLastAccess(NamespaceId ns)
         {
-            return _namespacePolicyResolver.GetPoliciesForNs(ns).LastAccessTracking;
+            if (!_namespacePolicyResolver.GetPoliciesForNs(ns).LastAccessTracking)
+            {
+                return false;
+            }
+
+            return _namespacePolicyResolver.GetPoliciesForNs(ns).GcMethod == NamespacePolicy.StoragePoolGCMethod.LastAccess;
         }
 
         public void Dispose()
diff --git a/Engine/Source/Programs/UnrealCloudDDC/Jupiter/Implementation/ObjectService.cs b/Engine/Source/Programs/UnrealCloudDDC/Jupiter/Implementation/ObjectService.cs
--- a/Engine/Source/Programs/UnrealCloudDDC/Jupiter/Implementation/ObjectService.cs
+++ b/Engine/Source/Programs/UnrealCloudDDC/Jupiter/Implementation/ObjectService.cs
@@ -7,6 +7,7 @@ using System.Threading.Tasks;
 using EpicGames.AspNet;
 using EpicGames.Horde.Storage;
 using EpicGames.Serialization;
+using Jupiter.Common;
 using Jupiter.Implementation.Blob;
 using Jupiter.Utils;
 using Microsoft.AspNetCore.Http;
@@ -39,11 +40,12 @@ namespace Jupiter.Implementation
         private readonly IReferenceResolver _referenceResolver;
         private readonly IReplicationLog _replicationLog;
         private readonly IBlobIndex _blobIndex;
+        private readonly INamespacePolicyResolver _namespacePolicyResolver;
         private readonly ILastAccessTracker<LastAccessRecord> _lastAccessTracker;
         private readonly Tracer _tracer;
         private readonly ILogger _logger;
 
-        public ObjectService(IHttpContextAccessor httpContextAccessor, IReferencesStore referencesStore, IBlobService blobService, IReferenceResolver referenceResolver, IReplicationLog replicationLog, IBlobIndex blobIndex, ILastAccessTracker<LastAccessRecord> lastAccessTracker, Tracer tracer, ILogger<ObjectService> logger)
+        public ObjectService(IHttpContextAccessor httpContextAccessor, IReferencesStore referencesStore, IBlobService blobService, IReferenceResolver referenceResolver, IReplicationLog replicationLog, IBlobIndex blobIndex, INamespacePolicyResolver namespacePolicyResolver, ILastAccessTracker<LastAccessRecord> lastAccessTracker, Tracer tracer, ILogger<ObjectService> logger)
         {
             _httpContextAccessor = httpContextAccessor;
             _referencesStore = referencesStore;
@@ -51,6 +53,7 @@ namespace Jupiter.Implementation
             _referenceResolver = referenceResolver;
             _replicationLog = replicationLog;
             _blobIndex = blobIndex;
+            _namespacePolicyResolver = namespacePolicyResolver;
             _lastAccessTracker = lastAccessTracker;
             _tracer = tracer;
             _logger = logger;
@@ -86,14 +89,17 @@ namespace Jupiter.Implementation
 
             if (doLastAccessTracking)
             {
-                // we do not wait for the last access tracking as it does not matter when it completes
-                Task lastAccessTask = _lastAccessTracker.TrackUsed(new LastAccessRecord(ns, bucket, key)).ContinueWith((task, _) =>
+                if (_namespacePolicyResolver.GetPoliciesForNs(ns).GcMethod == NamespacePolicy.StoragePoolGCMethod.LastAccess)
                 {
-                    if (task.Exception != null)
+                    // we do not wait for the last access tracking as it does not matter when it completes
+                    Task lastAccessTask = _lastAccessTracker.TrackUsed(new LastAccessRecord(ns, bucket, key)).ContinueWith((task, _) =>
                     {
-                        _logger.LogError(task.Exception, "Exception when tracking last access record");
-                    }
-                }, null, TaskScheduler.Current);
+                        if (task.Exception != null)
+                        {
+                            _logger.LogError(task.Exception, "Exception when tracking last access record");
+                        }
+                    }, null, TaskScheduler.Current);
+                }
             }
 
             BlobContents? blobContents = null;
diff --git a/Engine/Source/Programs/UnrealCloudDDC/Jupiter/Implementation/Objects/CacheReferencesStore.cs b/Engine/Source/Programs/UnrealCloudDDC/Jupiter/Implementation/Objects/CacheReferencesStore.cs
--- a/Engine/Source/Programs/UnrealCloudDDC/Jupiter/Implementation/Objects/CacheReferencesStore.cs
+++ b/Engine/Source/Programs/UnrealCloudDDC/Jupiter/Implementation/Objects/CacheReferencesStore.cs
@@ -25,7 +25,7 @@ namespace Jupiter.Implementation
         public CacheReferencesStore(IServiceProvider provider)
         {
             _mongoReferenceStore = ActivatorUtilities.CreateInstance<MongoReferencesStore>(provider, "Jupiter_Cache");
-            _mongoReferenceStore.SetTTLDuration(TimeSpan.FromDays(7.0));
+            _mongoReferenceStore.SetLastAccessTTLDuration(TimeSpan.FromDays(7.0));
 
             _upstreamReferenceStore = ActivatorUtilities.CreateInstance<UpstreamReferenceStore>(provider);
         }
diff --git a/Engine/Source/Programs/UnrealCloudDDC/Jupiter/Implementation/Objects/MongoReferencesStore.cs b/Engine/Source/Programs/UnrealCloudDDC/Jupiter/Implementation/Objects/MongoReferencesStore.cs
--- a/Engine/Source/Programs/UnrealCloudDDC/Jupiter/Implementation/Objects/MongoReferencesStore.cs
+++ b/Engine/Source/Programs/UnrealCloudDDC/Jupiter/Implementation/Objects/MongoReferencesStore.cs
@@ -4,6 +4,7 @@ using System;
 using System.Collections.Generic;
 using System.Threading.Tasks;
 using EpicGames.Horde.Storage;
+using Jupiter.Common;
 using Microsoft.Extensions.Options;
 using MongoDB.Bson.Serialization.Attributes;
 using MongoDB.Driver;
@@ -12,8 +13,11 @@ namespace Jupiter.Implementation
 {
     public class MongoReferencesStore : MongoStore, IReferencesStore
     {
-        public MongoReferencesStore(IOptionsMonitor<MongoSettings> settings, string? overrideDatabaseName = null) : base(settings, overrideDatabaseName)
+        private readonly INamespacePolicyResolver _namespacePolicyResolver;
+
+        public MongoReferencesStore(IOptionsMonitor<MongoSettings> settings, INamespacePolicyResolver namespacePolicyResolver, string? overrideDatabaseName = null) : base(settings, overrideDatabaseName)
         {
+            _namespacePolicyResolver = namespacePolicyResolver;
             CreateCollectionIfNotExists<MongoReferencesModelV0>().Wait();
             CreateCollectionIfNotExists<MongoNamespacesModelV0>().Wait();
 
@@ -36,6 +40,11 @@ namespace Jupiter.Implementation
                 indexModelNamespace
             });
 
+            CreateIndexModel<MongoReferencesModelV0> indexTTL = new CreateIndexModel<MongoReferencesModelV0>(
+                indexKeysDefinitionBuilder.Ascending(m => m.ExpireAt), new CreateIndexOptions {Name = "ExpireAtTTL", ExpireAfter = TimeSpan.Zero}
+            );
+
+            AddIndexFor<MongoReferencesModelV0>().CreateOne(indexTTL);
         }
 
         public async Task<ObjectRecord> Get(NamespaceId ns, BucketId bucket, IoHashKey key, IReferencesStore.FieldFlags flags)
@@ -72,6 +81,12 @@ namespace Jupiter.Implementation
             Task addNamespaceTask = AddNamespaceIfNotExist(ns);
             MongoReferencesModelV0 model = new MongoReferencesModelV0(ns, bucket, key, blobHash, blob, isFinalized, DateTime.Now);
             
+            NamespacePolicy policy = _namespacePolicyResolver.GetPoliciesForNs(ns);
+            if (policy.GcMethod == NamespacePolicy.StoragePoolGCMethod.TTL)
+            {
+                model.ExpireAt = DateTime.UtcNow.Add(policy.DefaultTTL);
+            }
+
             FilterDefinition<MongoReferencesModelV0> filter = Builders<MongoReferencesModelV0>.Filter.Where(m => m.Ns == ns.ToString() && m.Bucket == bucket.ToString() && m.Key == key.ToString());
             FindOneAndReplaceOptions<MongoReferencesModelV0, MongoReferencesModelV0> options = new FindOneAndReplaceOptions<MongoReferencesModelV0, MongoReferencesModelV0>
             {
@@ -185,7 +200,7 @@ namespace Jupiter.Implementation
             return 0L;
         }
 
-        public void SetTTLDuration(TimeSpan duration)
+        public void SetLastAccessTTLDuration(TimeSpan duration)
         {
             IndexKeysDefinitionBuilder<MongoReferencesModelV0> indexKeysDefinitionBuilder = Builders<MongoReferencesModelV0>.IndexKeys;
             CreateIndexModel<MongoReferencesModelV0> indexTTL = new CreateIndexModel<MongoReferencesModelV0>(
@@ -243,6 +258,7 @@ namespace Jupiter.Implementation
         public DateTime LastAccessTime { get; set; }
 
         public byte[]? InlineBlob { get; set; }
+        public DateTime ExpireAt { get; set; } = DateTime.MaxValue;
 
         public ObjectRecord ToObjectRecord()
         {
diff --git a/Engine/Source/Programs/UnrealCloudDDC/Jupiter/Implementation/Objects/ScyllaReferencesStore.cs b/Engine/Source/Programs/UnrealCloudDDC/Jupiter/Implementation/Objects/ScyllaReferencesStore.cs
--- a/Engine/Source/Programs/UnrealCloudDDC/Jupiter/Implementation/Objects/ScyllaReferencesStore.cs
+++ b/Engine/Source/Programs/UnrealCloudDDC/Jupiter/Implementation/Objects/ScyllaReferencesStore.cs
@@ -7,6 +7,7 @@ using System.Threading.Tasks;
 using Cassandra;
 using Cassandra.Mapping;
 using EpicGames.Horde.Storage;
+using Jupiter.Common;
 using Microsoft.Extensions.Logging;
 using Microsoft.Extensions.Options;
 using OpenTelemetry.Trace;
@@ -18,16 +19,18 @@ namespace Jupiter.Implementation
         private readonly ISession _session;
         private readonly IMapper _mapper;
         private readonly IOptionsMonitor<ScyllaSettings> _settings;
+        private readonly INamespacePolicyResolver _namespacePolicyResolver;
         private readonly Tracer _tracer;
         private readonly ILogger _logger;
         private readonly PreparedStatement _getObjectsStatement;
         private readonly PreparedStatement _getNamespacesStatement;
         private readonly PreparedStatement _getObjectsForPartitionRangeStatement;
 
-        public ScyllaReferencesStore(IScyllaSessionManager scyllaSessionManager, IOptionsMonitor<ScyllaSettings> settings, Tracer tracer, ILogger<ScyllaReferencesStore> logger)
+        public ScyllaReferencesStore(IScyllaSessionManager scyllaSessionManager, IOptionsMonitor<ScyllaSettings> settings, INamespacePolicyResolver namespacePolicyResolver, Tracer tracer, ILogger<ScyllaReferencesStore> logger)
         {
             _session = scyllaSessionManager.GetSessionForReplicatedKeyspace();
             _settings = settings;
+            _namespacePolicyResolver = namespacePolicyResolver;
             _tracer = tracer;
             _logger = logger;
 
@@ -107,7 +110,13 @@ namespace Jupiter.Implementation
             // add the bucket in parallel with inserting the actual object
             Task addBucketTask = MaybeAddBucket(ns, bucket);
 
-            await _mapper.InsertAsync<ScyllaObject>(new ScyllaObject(ns, bucket, name, blob, blobHash, isFinalized));
+            int? ttl = null;
+            NamespacePolicy policy = _namespacePolicyResolver.GetPoliciesForNs(ns);
+            if (policy.GcMethod == NamespacePolicy.StoragePoolGCMethod.TTL)
+            {
+                ttl = (int)policy.DefaultTTL.TotalSeconds;
+            }
+            await _mapper.InsertAsync<ScyllaObject>(new ScyllaObject(ns, bucket, name, blob, blobHash, isFinalized), ttl: ttl, insertNulls: false);
             await addBucketTask;
         }
 
diff --git a/Engine/Source/Programs/UnrealCloudDDC/Jupiter/Jupiter.Startup.cs b/Engine/Source/Programs/UnrealCloudDDC/Jupiter/Jupiter.Startup.cs
--- a/Engine/Source/Programs/UnrealCloudDDC/Jupiter/Jupiter.Startup.cs
+++ b/Engine/Source/Programs/UnrealCloudDDC/Jupiter/Jupiter.Startup.cs
@@ -178,7 +178,7 @@ namespace Jupiter
             services.AddSingleton(typeof(IPeerStatusService), typeof(PeerStatusService));
             services.AddHostedService<PeerStatusService>(p => (PeerStatusService)p.GetService<IPeerStatusService>()!);
         
-            services.AddTransient(typeof(IRefCleanup), typeof(RefCleanup));
+            services.AddTransient(typeof(IRefCleanup), typeof(RefLastAccessCleanup));
 
             services.AddTransient(typeof(VersionFile), typeof(VersionFile));
 
