Unreal Cloud DDC - Allow for uploading to the primary namespace to also populate the fallback namespace. This makes the permanent va storage invisible to the cooker and simplifies deployment on that side. Also fixed a bug were primary namespace would not be populated when moving data from the fallback namespace.

diff --git a/Engine/Source/Programs/UnrealCloudDDC/Jupiter.Common/BaseStartup.cs b/Engine/Source/Programs/UnrealCloudDDC/Jupiter.Common/BaseStartup.cs
--- a/Engine/Source/Programs/UnrealCloudDDC/Jupiter.Common/BaseStartup.cs
+++ b/Engine/Source/Programs/UnrealCloudDDC/Jupiter.Common/BaseStartup.cs
@@ -732,6 +732,7 @@ namespace Jupiter
         public bool UseBlobIndexForSlowExists { get; set; } = false;
         public bool IsPublicNamespace { get; set; } = true;
         public NamespaceId? FallbackNamespace { get; set; } = null;
+        public bool PopulateFallbackNamespaceOnUpload { get; set; } = true;
 
         public enum StoragePoolGCMethod  { None, LastAccess, TTL };
 
diff --git a/Engine/Source/Programs/UnrealCloudDDC/Jupiter/Implementation/Blob/IBlobService.cs b/Engine/Source/Programs/UnrealCloudDDC/Jupiter/Implementation/Blob/IBlobService.cs
--- a/Engine/Source/Programs/UnrealCloudDDC/Jupiter/Implementation/Blob/IBlobService.cs
+++ b/Engine/Source/Programs/UnrealCloudDDC/Jupiter/Implementation/Blob/IBlobService.cs
@@ -206,7 +206,11 @@ public class BlobService : IBlobService
             await using Stream s = bufferedPayload.GetStream();
             await store.PutObject(ns, s, identifier);
         }
-
+        NamespacePolicy policy = _namespacePolicyResolver.GetPoliciesForNs(ns);
+        if (policy.PopulateFallbackNamespaceOnUpload && policy.FallbackNamespace.HasValue)
+        {
+            await PutObjectToStores(policy.FallbackNamespace.Value, bufferedPayload, identifier);
+        }
         return identifier;
     }
 
@@ -217,6 +221,11 @@ public class BlobService : IBlobService
             await store.PutObject(ns, payload, identifier);
         }
 
+        NamespacePolicy policy = _namespacePolicyResolver.GetPoliciesForNs(ns);
+        if (policy.PopulateFallbackNamespaceOnUpload && policy.FallbackNamespace.HasValue)
+        {
+            await PutObjectToStores(policy.FallbackNamespace.Value, payload, identifier);
+        }
         return identifier;
     }
 
@@ -256,8 +265,21 @@ public class BlobService : IBlobService
 
                 try
                 {
-                    return await GetObjectFromStores(policy.FallbackNamespace.Value, blob, storageLayers);
+                    IServerTiming? serverTiming = _httpContextAccessor.HttpContext?.RequestServices.GetService<IServerTiming>();
+
+                    // read the content from the fallback namespace
+                    BlobContents fallbackContent = await GetObjectFromStores(policy.FallbackNamespace.Value, blob, storageLayers);
+                    
+                    // populate the primary namespace with the content
+                    using TelemetrySpan _ = _tracer.StartActiveSpan("HierarchicalStore.Populate").SetAttribute("operation.name", "HierarchicalStore.Populate");
+                    using ServerTimingMetricScoped? serverTimingScope = serverTiming?.CreateServerTimingMetricScope($"blob.populate", "Populating caches with blob contents");
 
+                    await using MemoryStream tempStream = new MemoryStream();
+                    await fallbackContent.Stream.CopyToAsync(tempStream);
+                    byte[] data = tempStream.ToArray();
+
+                    await PutObject(ns, data, blob);
+                    return await GetObject(ns, blob);
                 }
                 catch (BlobNotFoundException)
                 {
@@ -277,6 +299,19 @@ public class BlobService : IBlobService
                 }
             }
 
+            // if the primary namespace failed check to see if we should use a fallback policy which has replication enabled
+            if (policy.FallbackNamespace != null && ShouldFetchBlobOnDemand(policy.FallbackNamespace.Value))
+            {
+                try
+                {
+                    return await ReplicateObject(policy.FallbackNamespace.Value, blob);
+                }
+                catch (BlobNotFoundException)
+                {
+                    // if the blob is not found we can ignore it as we will just rethrow it later anyway
+                }
+            }
+
             // we might have attempted to fetch the object and failed, or had no fallback options, either way the blob can not be found so we should rethrow
             throw;
         }
