Unreal Cloud DDC - Enable on demand replication feature toggle by default and make it configurable by the configmap. This still requires opt in per namespace which isn't set by default.

diff --git a/Engine/Source/Programs/UnrealCloudDDC/Helm/UnrealCloudDDC/templates/configmap.yaml b/Engine/Source/Programs/UnrealCloudDDC/Helm/UnrealCloudDDC/templates/configmap.yaml
--- a/Engine/Source/Programs/UnrealCloudDDC/Helm/UnrealCloudDDC/templates/configmap.yaml
+++ b/Engine/Source/Programs/UnrealCloudDDC/Helm/UnrealCloudDDC/templates/configmap.yaml
@@ -29,6 +29,10 @@ UnrealCloudDDC:
   ContentIdStoreImplementation :  {{ default "Scylla" $config.ContentIdStoreImplementation }}
   BlobIndexImplementation :  {{ default "Scylla" $config.BlobIndexImplementation }}
 
+{{- if $config.EnableOnDemandReplication -}}
+  EnableOnDemandReplication :  {{ $config.EnableOnDemandReplication }}
+{{- end -}}
+
 {{- if $config.EnableLastAccessTracking -}}
   EnableLastAccessTracking: {{ $config.EnableLastAccessTracking }}
 {{- end -}}
diff --git a/Engine/Source/Programs/UnrealCloudDDC/Jupiter/Settings.cs b/Engine/Source/Programs/UnrealCloudDDC/Jupiter/Settings.cs
--- a/Engine/Source/Programs/UnrealCloudDDC/Jupiter/Settings.cs
+++ b/Engine/Source/Programs/UnrealCloudDDC/Jupiter/Settings.cs
@@ -111,7 +111,7 @@ namespace Jupiter
 
         public int LastAccessRollupFrequencySeconds { get; set; } = 900; // 15 minutes
         public bool EnableLastAccessTracking { get; set; } = true;
-        public bool EnableOnDemandReplication { get; set; } = false;
+        public bool EnableOnDemandReplication { get; set; } = true;
     }
 
     public class MongoSettings
