Unreal Cloud DDC - Fixed issue were fallback namespaces were not having its blob index populated as data was written to it.

diff --git a/Engine/Source/Programs/UnrealCloudDDC/Jupiter/Implementation/Blob/IBlobService.cs b/Engine/Source/Programs/UnrealCloudDDC/Jupiter/Implementation/Blob/IBlobService.cs
--- a/Engine/Source/Programs/UnrealCloudDDC/Jupiter/Implementation/Blob/IBlobService.cs
+++ b/Engine/Source/Programs/UnrealCloudDDC/Jupiter/Implementation/Blob/IBlobService.cs
@@ -210,6 +210,7 @@ public class BlobService : IBlobService
         if (policy.PopulateFallbackNamespaceOnUpload && policy.FallbackNamespace.HasValue)
         {
             await PutObjectToStores(policy.FallbackNamespace.Value, bufferedPayload, identifier);
+            await _blobIndex.AddBlobToIndex(policy.FallbackNamespace.Value, identifier);
         }
         return identifier;
     }
@@ -225,6 +226,7 @@ public class BlobService : IBlobService
         if (policy.PopulateFallbackNamespaceOnUpload && policy.FallbackNamespace.HasValue)
         {
             await PutObjectToStores(policy.FallbackNamespace.Value, payload, identifier);
+            await _blobIndex.AddBlobToIndex(policy.FallbackNamespace.Value, identifier);
         }
         return identifier;
     }
