Unreal Cloud DDC - Added metric for replication - can hopefully replace the excessive logging we use to monitor that replication is running

diff --git a/Engine/Source/Programs/UnrealCloudDDC/Jupiter/Implementation/Replication/RefsReplicator.cs b/Engine/Source/Programs/UnrealCloudDDC/Jupiter/Implementation/Replication/RefsReplicator.cs
--- a/Engine/Source/Programs/UnrealCloudDDC/Jupiter/Implementation/Replication/RefsReplicator.cs
+++ b/Engine/Source/Programs/UnrealCloudDDC/Jupiter/Implementation/Replication/RefsReplicator.cs
@@ -2,6 +2,7 @@
 
 using System;
 using System.Collections.Generic;
+using System.Diagnostics.Metrics;
 using System.Globalization;
 using System.IO;
 using System.Linq;
@@ -42,7 +43,9 @@ namespace Jupiter.Implementation
         private bool _replicationRunning;
         private bool _disposed = false;
 
-        public RefsReplicator(ReplicatorSettings replicatorSettings, IBlobService blobService, IHttpClientFactory httpClientFactory, IReplicationLog replicationLog, IServiceCredentials serviceCredentials, Tracer tracer, BufferedPayloadFactory bufferedPayloadFactory, ReplicationLogFactory replicationLogFactory, ILogger<RefsReplicator> logger)
+        private static Histogram<long>? s_replicatedCounter;
+
+        public RefsReplicator(ReplicatorSettings replicatorSettings, IBlobService blobService, IHttpClientFactory httpClientFactory, IReplicationLog replicationLog, IServiceCredentials serviceCredentials, Tracer tracer, BufferedPayloadFactory bufferedPayloadFactory, ReplicationLogFactory replicationLogFactory, ILogger<RefsReplicator> logger, Meter meter)
         {
             _name = replicatorSettings.ReplicatorName;
             _namespace = new NamespaceId(replicatorSettings.NamespaceToReplicate);
@@ -73,6 +76,11 @@ namespace Jupiter.Implementation
             }
 
             Info = new ReplicatorInfo(replicatorSettings.ReplicatorName, _namespace, _refsState);
+
+            if (s_replicatedCounter == null)
+            {
+                s_replicatedCounter = meter.CreateHistogram<long>("replication.active");
+            }
         }
 
         public void Dispose()
@@ -613,6 +621,8 @@ namespace Jupiter.Implementation
 
         private void LogReplicationHeartbeat(int countOfCurrentReplications)
         {
+            s_replicatedCounter?.Record(countOfCurrentReplications, new KeyValuePair<string, object?>("replicator", _name), new KeyValuePair<string, object?>("namespace", _namespace));
+
             // log message used to generate metric for how many replications are currently running
             _logger.LogInformation("{Name} replication has run . Count of running replications: {CurrentReplications}", _name, countOfCurrentReplications);
 
