Unreal Cloud DDC - Ability to override S3 bucket name used for a storage pool

Overriding the s3 bucket name is similar to what we already allow for Azure blob storages, this allows you more control over exactly which bucket is used for which storage pool. We will use this to store all permanent VA assets in a single bucket, resulting in less redundant data across regions but with slower access to them.

diff --git a/Engine/Source/Programs/UnrealCloudDDC/Jupiter/Implementation/Blob/AmazonS3Store.cs b/Engine/Source/Programs/UnrealCloudDDC/Jupiter/Implementation/Blob/AmazonS3Store.cs
--- a/Engine/Source/Programs/UnrealCloudDDC/Jupiter/Implementation/Blob/AmazonS3Store.cs
+++ b/Engine/Source/Programs/UnrealCloudDDC/Jupiter/Implementation/Blob/AmazonS3Store.cs
@@ -10,7 +10,6 @@ using Amazon.S3.Model;
 using EpicGames.Horde.Storage;
 using Jupiter.Implementation.Blob;
 using Jupiter.Common;
-using Jupiter.Common.Implementation;
 using Microsoft.Extensions.Options;
 using KeyNotFoundException = System.Collections.Generic.KeyNotFoundException;
 using System.Threading;
@@ -63,7 +62,14 @@ namespace Jupiter.Implementation
         {
             try
             {
-                string storagePool = _namespacePolicyResolver.GetPoliciesForNs(ns).StoragePool;
+                NamespacePolicy policy = _namespacePolicyResolver.GetPoliciesForNs(ns);
+                string storagePool = policy.StoragePool;
+                // if the bucket to use for the storage pool has been overriden we use the override
+                if (_settings.StoragePoolBucketOverride.TryGetValue(storagePool, out string? containerOverride))
+                {
+                    return containerOverride;
+                }
+                // by default we use the storage pool as a suffix to determine the bucket for that pool
                 string storagePoolSuffix = string.IsNullOrEmpty(storagePool) ? "" : $"-{storagePool}";
                 return $"{_settings.BucketName}{storagePoolSuffix}";
             }
diff --git a/Engine/Source/Programs/UnrealCloudDDC/Jupiter/Settings.cs b/Engine/Source/Programs/UnrealCloudDDC/Jupiter/Settings.cs
--- a/Engine/Source/Programs/UnrealCloudDDC/Jupiter/Settings.cs
+++ b/Engine/Source/Programs/UnrealCloudDDC/Jupiter/Settings.cs
@@ -173,6 +173,10 @@ namespace Jupiter
         public bool SetBucketPolicies { get; set; } = true;
 
         public bool UseBlobIndexForExistsCheck { get; set; } = false;
+
+        [System.Diagnostics.CodeAnalysis.SuppressMessage("Usage", "CA2227:Collection properties should be read only", Justification = "Modified by settings")]
+        // ReSharper disable once CollectionNeverUpdated.Global
+        public Dictionary<string, string> StoragePoolBucketOverride { get; set; } = new Dictionary<string, string>();
     }
 
     public class GCSettings
