Unreal Cloud DDC - Fixed issue were s3 storage pool overrides were not written to the configmap

diff --git a/Engine/Source/Programs/UnrealCloudDDC/Helm/UnrealCloudDDC/templates/configmap.yaml b/Engine/Source/Programs/UnrealCloudDDC/Helm/UnrealCloudDDC/templates/configmap.yaml
--- a/Engine/Source/Programs/UnrealCloudDDC/Helm/UnrealCloudDDC/templates/configmap.yaml
+++ b/Engine/Source/Programs/UnrealCloudDDC/Helm/UnrealCloudDDC/templates/configmap.yaml
@@ -141,5 +141,7 @@ S3:
   SetBucketPolicies: {{ default false $s3.SetBucketPolicies }}
 
   UseBlobIndexForExistsCheck: {{ default false $s3.UseBlobIndexForExistsCheck }}
+  StoragePoolBucketOverride: {{ default {} $s3.StoragePoolBucketOverride }}
+
 {{ end }}
 
