Unreal Cloud DDC - Added ability to configure unknown namespaces to be cleaned up during ref GC

diff --git a/Engine/Source/Programs/UnrealCloudDDC/Jupiter.Common/BaseStartup.cs b/Engine/Source/Programs/UnrealCloudDDC/Jupiter.Common/BaseStartup.cs
--- a/Engine/Source/Programs/UnrealCloudDDC/Jupiter.Common/BaseStartup.cs
+++ b/Engine/Source/Programs/UnrealCloudDDC/Jupiter.Common/BaseStartup.cs
@@ -734,7 +734,24 @@ namespace Jupiter
         public NamespaceId? FallbackNamespace { get; set; } = null;
         public bool PopulateFallbackNamespaceOnUpload { get; set; } = true;
 
-        public enum StoragePoolGCMethod  { None, LastAccess, TTL };
+        public enum StoragePoolGCMethod  {
+            /// <summary>
+            /// Never run GC on this namespace
+            /// </summary>
+            None, 
+            /// <summary>
+            /// Apply last access deletion, objects not used for a duration set in GCSettings will be removed
+            /// </summary>
+            LastAccess, 
+            /// <summary>
+            /// Objects are removed after DefaultTTL time has passed, no matter if they are used or not
+            /// </summary>
+            TTL,
+            /// <summary>
+            /// Always GC references to this namespace, used to opt in to cleaning out old data
+            /// </summary>
+            Always
+        };
 
         public StoragePoolGCMethod GcMethod { get; set; } = StoragePoolGCMethod.LastAccess;
 
diff --git a/Engine/Source/Programs/UnrealCloudDDC/Jupiter/Implementation/GC/CacheCleanup.cs b/Engine/Source/Programs/UnrealCloudDDC/Jupiter/Implementation/GC/CacheCleanup.cs
--- a/Engine/Source/Programs/UnrealCloudDDC/Jupiter/Implementation/GC/CacheCleanup.cs
+++ b/Engine/Source/Programs/UnrealCloudDDC/Jupiter/Implementation/GC/CacheCleanup.cs
@@ -46,14 +46,28 @@ namespace Jupiter.Implementation
                 return false;
             }
 
-            NamespacePolicy policy = _namespacePolicyResolver.GetPoliciesForNs(ns);
+            try
+            {
+                NamespacePolicy policy = _namespacePolicyResolver.GetPoliciesForNs(ns);
+
+                if (policy.GcMethod == NamespacePolicy.StoragePoolGCMethod.LastAccess)
+                {
+                    // only run for namespaces set to use last access tracking
+                    return true;
+                }
 
-            if (policy.GcMethod == NamespacePolicy.StoragePoolGCMethod.LastAccess)
+                if (policy.GcMethod == NamespacePolicy.StoragePoolGCMethod.Always)
+                {
+                    // this is a old namespace that should be cleaned up
+                    return true;
+                }
+            }
+            catch (NamespaceNotFoundException)
             {
-                // only run for namespaces set to use last access tracking
-                return true;
+                _logger.LogWarning("Unknown namespace {Namespace} when attempting to GC References. To opt in to deleting the old namespace add a policy for it with the GcMethod set to always", ns);
+                return false;
             }
-
+            
             return false;
         }
 
