Unreal Cloud DDC - Fixed incorrect default for storage pool override in s3 for helm charts

diff --git a/Engine/Source/Programs/UnrealCloudDDC/Helm/UnrealCloudDDC/templates/configmap.yaml b/Engine/Source/Programs/UnrealCloudDDC/Helm/UnrealCloudDDC/templates/configmap.yaml
--- a/Engine/Source/Programs/UnrealCloudDDC/Helm/UnrealCloudDDC/templates/configmap.yaml
+++ b/Engine/Source/Programs/UnrealCloudDDC/Helm/UnrealCloudDDC/templates/configmap.yaml
@@ -141,7 +141,7 @@ S3:
   SetBucketPolicies: {{ default false $s3.SetBucketPolicies }}
 
   UseBlobIndexForExistsCheck: {{ default false $s3.UseBlobIndexForExistsCheck }}
-  StoragePoolBucketOverride: {{ default {} $s3.StoragePoolBucketOverride }}
+  StoragePoolBucketOverride: {{ default dict $s3.StoragePoolBucketOverride }}
 
 {{ end }}
 
