Unreal Cloud DDC - Fixed s3 override config not getting converted to yaml in helm configmap

diff --git a/Engine/Source/Programs/UnrealCloudDDC/Helm/UnrealCloudDDC/templates/configmap.yaml b/Engine/Source/Programs/UnrealCloudDDC/Helm/UnrealCloudDDC/templates/configmap.yaml
--- a/Engine/Source/Programs/UnrealCloudDDC/Helm/UnrealCloudDDC/templates/configmap.yaml
+++ b/Engine/Source/Programs/UnrealCloudDDC/Helm/UnrealCloudDDC/templates/configmap.yaml
@@ -141,7 +141,8 @@ S3:
   SetBucketPolicies: {{ default false $s3.SetBucketPolicies }}
 
   UseBlobIndexForExistsCheck: {{ default false $s3.UseBlobIndexForExistsCheck }}
-  StoragePoolBucketOverride: {{ default dict $s3.StoragePoolBucketOverride }}
-
+  {{- if $s3.StoragePoolBucketOverride }}
+  StoragePoolBucketOverride: {{ $s3.StoragePoolBucketOverride | toYaml | nindent 4  }}
+  {{- end }}
 {{ end }}
 
