Unreal Cloud DDC - Added UseArnRegion option to allow for cross region s3 calls

diff --git a/Engine/Source/Programs/UnrealCloudDDC/Helm/UnrealCloudDDC/templates/configmap.yaml b/Engine/Source/Programs/UnrealCloudDDC/Helm/UnrealCloudDDC/templates/configmap.yaml
--- a/Engine/Source/Programs/UnrealCloudDDC/Helm/UnrealCloudDDC/templates/configmap.yaml
+++ b/Engine/Source/Programs/UnrealCloudDDC/Helm/UnrealCloudDDC/templates/configmap.yaml
@@ -141,6 +141,11 @@ S3:
   SetBucketPolicies: {{ default false $s3.SetBucketPolicies }}
 
   UseBlobIndexForExistsCheck: {{ default false $s3.UseBlobIndexForExistsCheck }}
+
+  {{- if $s3.UseArnRegion }}
+  UseArnRegion: {{ $s3.UseArnRegion.UseArnRegion}}
+  {{- end }}
+
   {{- if $s3.StoragePoolBucketOverride }}
   StoragePoolBucketOverride: {{ $s3.StoragePoolBucketOverride | toYaml | nindent 4  }}
   {{- end }}
diff --git a/Engine/Source/Programs/UnrealCloudDDC/Jupiter/Jupiter.Startup.cs b/Engine/Source/Programs/UnrealCloudDDC/Jupiter/Jupiter.Startup.cs
--- a/Engine/Source/Programs/UnrealCloudDDC/Jupiter/Jupiter.Startup.cs
+++ b/Engine/Source/Programs/UnrealCloudDDC/Jupiter/Jupiter.Startup.cs
@@ -473,6 +473,11 @@ namespace Jupiter
                     {
                         c.ForcePathStyle = true;
                     }
+
+                    if (s3Settings.UseArnRegion.HasValue)
+                    {
+                        c.UseArnRegion = s3Settings.UseArnRegion.Value;
+                    }
                 }
                 else
                 {
diff --git a/Engine/Source/Programs/UnrealCloudDDC/Jupiter/Settings.cs b/Engine/Source/Programs/UnrealCloudDDC/Jupiter/Settings.cs
--- a/Engine/Source/Programs/UnrealCloudDDC/Jupiter/Settings.cs
+++ b/Engine/Source/Programs/UnrealCloudDDC/Jupiter/Settings.cs
@@ -177,6 +177,8 @@ namespace Jupiter
         [System.Diagnostics.CodeAnalysis.SuppressMessage("Usage", "CA2227:Collection properties should be read only", Justification = "Modified by settings")]
         // ReSharper disable once CollectionNeverUpdated.Global
         public Dictionary<string, string> StoragePoolBucketOverride { get; set; } = new Dictionary<string, string>();
+
+        public bool? UseArnRegion { get; set; } = null;
     }
 
     public class GCSettings
