Unreal Cloud DDC - Fixed copy paste error when resolving UseArnRegion

diff --git a/Engine/Source/Programs/UnrealCloudDDC/Helm/UnrealCloudDDC/templates/configmap.yaml b/Engine/Source/Programs/UnrealCloudDDC/Helm/UnrealCloudDDC/templates/configmap.yaml
--- a/Engine/Source/Programs/UnrealCloudDDC/Helm/UnrealCloudDDC/templates/configmap.yaml
+++ b/Engine/Source/Programs/UnrealCloudDDC/Helm/UnrealCloudDDC/templates/configmap.yaml
@@ -143,7 +143,7 @@ S3:
   UseBlobIndexForExistsCheck: {{ default false $s3.UseBlobIndexForExistsCheck }}
 
   {{- if $s3.UseArnRegion }}
-  UseArnRegion: {{ $s3.UseArnRegion.UseArnRegion}}
+  UseArnRegion: {{ $s3.UseArnRegion }}
   {{- end }}
 
   {{- if $s3.StoragePoolBucketOverride }}
