Unreal Cloud DDC - Removed duplicate exceptions that both signaled namespace not existing. Fixes issue with undefined namespaces in ref GC.

diff --git a/Engine/Source/Programs/UnrealCloudDDC/Jupiter.Common/Authentication/NamespaceAuthorizationHandler.cs b/Engine/Source/Programs/UnrealCloudDDC/Jupiter.Common/Authentication/NamespaceAuthorizationHandler.cs
--- a/Engine/Source/Programs/UnrealCloudDDC/Jupiter.Common/Authentication/NamespaceAuthorizationHandler.cs
+++ b/Engine/Source/Programs/UnrealCloudDDC/Jupiter.Common/Authentication/NamespaceAuthorizationHandler.cs
@@ -73,7 +73,7 @@ namespace Jupiter
                     context.Succeed(requirement);
                 }
             }
-            catch (UnknownNamespaceException)
+            catch (NamespaceNotFoundException)
             {
                 // if the namespace doesn't have a policy setup, e.g. we do not know which claims to require then we can just exit here as the auth will fail
             }
diff --git a/Engine/Source/Programs/UnrealCloudDDC/Jupiter.Common/Implementation/NamespacePolicyResolver.cs b/Engine/Source/Programs/UnrealCloudDDC/Jupiter.Common/Implementation/NamespacePolicyResolver.cs
--- a/Engine/Source/Programs/UnrealCloudDDC/Jupiter.Common/Implementation/NamespacePolicyResolver.cs
+++ b/Engine/Source/Programs/UnrealCloudDDC/Jupiter.Common/Implementation/NamespacePolicyResolver.cs
@@ -65,14 +65,22 @@ namespace Jupiter.Common
                 return _internalNamespacePolicy;
             }
 
-            throw new UnknownNamespaceException($"Unable to find a valid policy for namespace {ns}");
+            throw new NamespaceNotFoundException(ns, $"Unable to find a valid policy for namespace {ns}");
         }
     }
 
-    public class UnknownNamespaceException : Exception
+    public class NamespaceNotFoundException : Exception
     {
-        public UnknownNamespaceException(string message) : base(message)
+        public NamespaceId Namespace { get; }
+
+        public NamespaceNotFoundException(NamespaceId @namespace) : base($"Could not find namespace {@namespace}")
+        {
+            Namespace = @namespace;
+        }
+
+        public NamespaceNotFoundException(NamespaceId @namespace, string message) : base(message)
         {
+            Namespace = @namespace;
         }
     }
 }
diff --git a/Engine/Source/Programs/UnrealCloudDDC/Jupiter/Controllers/ReferencesController.cs b/Engine/Source/Programs/UnrealCloudDDC/Jupiter/Controllers/ReferencesController.cs
--- a/Engine/Source/Programs/UnrealCloudDDC/Jupiter/Controllers/ReferencesController.cs
+++ b/Engine/Source/Programs/UnrealCloudDDC/Jupiter/Controllers/ReferencesController.cs
@@ -19,6 +19,7 @@ using EpicGames.Horde.Storage;
 using EpicGames.Serialization;
 using Jupiter.Implementation;
 using JetBrains.Annotations;
+using Jupiter.Common;
 using Jupiter.Common.Implementation;
 using Jupiter.Utils;
 using Microsoft.AspNetCore.Authorization;
@@ -774,6 +775,10 @@ namespace Jupiter.Controllers
             {
                 return BadRequest(e.Message);
             }
+            catch (ObjectNotFoundException e)
+            {
+                return NotFound(e.Message);
+            }
         }
 
         [HttpPost("{ns}")]
diff --git a/Engine/Source/Programs/UnrealCloudDDC/Jupiter/Controllers/ReplicationLogController.cs b/Engine/Source/Programs/UnrealCloudDDC/Jupiter/Controllers/ReplicationLogController.cs
--- a/Engine/Source/Programs/UnrealCloudDDC/Jupiter/Controllers/ReplicationLogController.cs
+++ b/Engine/Source/Programs/UnrealCloudDDC/Jupiter/Controllers/ReplicationLogController.cs
@@ -8,6 +8,7 @@ using System.Threading;
 using System.Threading.Tasks;
 using EpicGames.Horde.Storage;
 using EpicGames.Serialization;
+using Jupiter.Common;
 using Jupiter.Implementation;
 using Jupiter.Implementation.TransactionLog;
 using Microsoft.AspNetCore.Authorization;
diff --git a/Engine/Source/Programs/UnrealCloudDDC/Jupiter/Implementation/Blob/BlobIndexConsistencyCheckService.cs b/Engine/Source/Programs/UnrealCloudDDC/Jupiter/Implementation/Blob/BlobIndexConsistencyCheckService.cs
--- a/Engine/Source/Programs/UnrealCloudDDC/Jupiter/Implementation/Blob/BlobIndexConsistencyCheckService.cs
+++ b/Engine/Source/Programs/UnrealCloudDDC/Jupiter/Implementation/Blob/BlobIndexConsistencyCheckService.cs
@@ -167,7 +167,7 @@ namespace Jupiter.Implementation
                             }
                         }
                     }
-                    catch (UnknownNamespaceException)
+                    catch (NamespaceNotFoundException)
                     {
                         if (_settings.CurrentValue.AllowDeletesInBlobIndex)
                         {
diff --git a/Engine/Source/Programs/UnrealCloudDDC/Jupiter/Implementation/Blob/IBlobStore.cs b/Engine/Source/Programs/UnrealCloudDDC/Jupiter/Implementation/Blob/IBlobStore.cs
--- a/Engine/Source/Programs/UnrealCloudDDC/Jupiter/Implementation/Blob/IBlobStore.cs
+++ b/Engine/Source/Programs/UnrealCloudDDC/Jupiter/Implementation/Blob/IBlobStore.cs
@@ -67,16 +67,6 @@ namespace Jupiter.Implementation
         }
     }
 
-    public class NamespaceNotFoundException : Exception
-    {
-        public NamespaceId Namespace { get; }
-
-        public NamespaceNotFoundException(NamespaceId @namespace) : base($"Could not find namespace {@namespace}")
-        {
-            Namespace = @namespace;
-        }
-    }
-
     public class BlobToLargeException : Exception
     {
         public BlobIdentifier Blob { get; }
diff --git a/Engine/Source/Programs/UnrealCloudDDC/Jupiter/Implementation/Blob/MemoryBlobStore.cs b/Engine/Source/Programs/UnrealCloudDDC/Jupiter/Implementation/Blob/MemoryBlobStore.cs
--- a/Engine/Source/Programs/UnrealCloudDDC/Jupiter/Implementation/Blob/MemoryBlobStore.cs
+++ b/Engine/Source/Programs/UnrealCloudDDC/Jupiter/Implementation/Blob/MemoryBlobStore.cs
@@ -9,6 +9,7 @@ using System.Runtime.CompilerServices;
 using System.Threading;
 using System.Threading.Tasks;
 using EpicGames.Horde.Storage;
+using Jupiter.Common;
 using Jupiter.Utils;
 
 namespace Jupiter.Implementation
diff --git a/Engine/Source/Programs/UnrealCloudDDC/Jupiter/Implementation/TransactionLog/MemoryReplicationLog.cs b/Engine/Source/Programs/UnrealCloudDDC/Jupiter/Implementation/TransactionLog/MemoryReplicationLog.cs
--- a/Engine/Source/Programs/UnrealCloudDDC/Jupiter/Implementation/TransactionLog/MemoryReplicationLog.cs
+++ b/Engine/Source/Programs/UnrealCloudDDC/Jupiter/Implementation/TransactionLog/MemoryReplicationLog.cs
@@ -8,6 +8,7 @@ using System.Threading.Tasks;
 using Cassandra;
 using EpicGames.Core;
 using EpicGames.Horde.Storage;
+using Jupiter.Common;
 
 namespace Jupiter.Implementation
 {
diff --git a/Engine/Source/Programs/UnrealCloudDDC/Jupiter/Implementation/TransactionLog/ScyllaReplicationLog.cs b/Engine/Source/Programs/UnrealCloudDDC/Jupiter/Implementation/TransactionLog/ScyllaReplicationLog.cs
--- a/Engine/Source/Programs/UnrealCloudDDC/Jupiter/Implementation/TransactionLog/ScyllaReplicationLog.cs
+++ b/Engine/Source/Programs/UnrealCloudDDC/Jupiter/Implementation/TransactionLog/ScyllaReplicationLog.cs
@@ -7,6 +7,7 @@ using System.Threading.Tasks;
 using Cassandra;
 using Cassandra.Mapping;
 using EpicGames.Horde.Storage;
+using Jupiter.Common;
 using Microsoft.Extensions.Options;
 using OpenTelemetry.Trace;
 
