Unreal Cloud DDC - Reduced logging in worker now that we have moved the monitoring to use the metric generated instead.

diff --git a/Engine/Source/Programs/UnrealCloudDDC/Jupiter/Implementation/Replication/RefsReplicator.cs b/Engine/Source/Programs/UnrealCloudDDC/Jupiter/Implementation/Replication/RefsReplicator.cs
--- a/Engine/Source/Programs/UnrealCloudDDC/Jupiter/Implementation/Replication/RefsReplicator.cs
+++ b/Engine/Source/Programs/UnrealCloudDDC/Jupiter/Implementation/Replication/RefsReplicator.cs
@@ -623,11 +623,8 @@ namespace Jupiter.Implementation
         {
             s_replicatedCounter?.Record(countOfCurrentReplications, new KeyValuePair<string, object?>("replicator", _name), new KeyValuePair<string, object?>("namespace", _namespace));
 
-            // log message used to generate metric for how many replications are currently running
-            _logger.LogInformation("{Name} replication has run . Count of running replications: {CurrentReplications}", _name, countOfCurrentReplications);
-
             // log message used to verify replicators are actually running
-            _logger.LogInformation("{Name} starting replication. Last transaction was {TransactionId} {Generation}", _name, State.ReplicatorOffset.GetValueOrDefault(0L), State.ReplicatingGeneration.GetValueOrDefault(Guid.Empty) );
+            _logger.LogDebug("{Name} starting replication. Last transaction was {TransactionId} {Generation}. Count Of running replications: {CurrentReplications}", _name, State.ReplicatorOffset.GetValueOrDefault(0L), State.ReplicatingGeneration.GetValueOrDefault(Guid.Empty), countOfCurrentReplications);
         }
 
         private async Task AddToReplicationLog(NamespaceId ns, BucketId bucket, IoHashKey key, BlobIdentifier blob)
