Unreal Cloud DDC - Reduced level of more logging in the replicator for lines that are not particularly important.

diff --git a/Engine/Source/Programs/UnrealCloudDDC/Jupiter/Implementation/Replication/RefsReplicator.cs b/Engine/Source/Programs/UnrealCloudDDC/Jupiter/Implementation/Replication/RefsReplicator.cs
--- a/Engine/Source/Programs/UnrealCloudDDC/Jupiter/Implementation/Replication/RefsReplicator.cs
+++ b/Engine/Source/Programs/UnrealCloudDDC/Jupiter/Implementation/Replication/RefsReplicator.cs
@@ -393,7 +393,7 @@ namespace Jupiter.Implementation
                         .SetAttribute("resource.name", $"[REDACTED_EMAIL]}[REDACTED_EMAIL]}")
                         .SetAttribute("time-bucket", @event.Timestamp.ToString(CultureInfo.InvariantCulture));
 
-                    _logger.LogInformation("{Name} New transaction to replicate found. Ref: {Namespace} {Bucket} {Key} in {TimeBucket} ({TimeDate}) with id {EventId}. Count of running replications: {CurrentReplications}", _name, @event.Namespace, @event.Bucket, @event.Key, @event.TimeBucket, @event.Timestamp, @event.EventId, replicationTasks.Count);
+                    _logger.LogDebug("{Name} New transaction to replicate found. Ref: {Namespace} {Bucket} {Key} in {TimeBucket} ({TimeDate}) with id {EventId}. Count of running replications: {CurrentReplications}", _name, @event.Namespace, @event.Bucket, @event.Key, @event.TimeBucket, @event.Timestamp, @event.EventId, replicationTasks.Count);
                 
                     Info.CountOfRunningReplications = replicationTasks.Count;
                     LogReplicationHeartbeat(replicationTasks.Count);
