Unreal Cloud DDC - Added option to disable ref cache and set it for workers to see if this is why they are running out of memory.

diff --git a/Engine/Source/Programs/UnrealCloudDDC/Helm/UnrealCloudDDC/templates/configmap-worker.yaml b/Engine/Source/Programs/UnrealCloudDDC/Helm/UnrealCloudDDC/templates/configmap-worker.yaml
--- a/Engine/Source/Programs/UnrealCloudDDC/Helm/UnrealCloudDDC/templates/configmap-worker.yaml
+++ b/Engine/Source/Programs/UnrealCloudDDC/Helm/UnrealCloudDDC/templates/configmap-worker.yaml
@@ -100,6 +100,11 @@ ConsistencyCheck:
 GC: 
 {{ $config.GC | toYaml | nindent 2 }}
 
+{{- if $config.CacheRef }}
+CacheRef:
+{{ $config.CacheRef | toYaml | nindent 2 }}
+{{- end }}
+
 {{- if or ($worker.persistence.enabled) ($worker.persistence.volume) }}
   RunFilesystemCleanup: true
 {{- end }}
diff --git a/Engine/Source/Programs/UnrealCloudDDC/Helm/UnrealCloudDDC/templates/configmap.yaml b/Engine/Source/Programs/UnrealCloudDDC/Helm/UnrealCloudDDC/templates/configmap.yaml
--- a/Engine/Source/Programs/UnrealCloudDDC/Helm/UnrealCloudDDC/templates/configmap.yaml
+++ b/Engine/Source/Programs/UnrealCloudDDC/Helm/UnrealCloudDDC/templates/configmap.yaml
@@ -113,6 +113,11 @@ Serilog:
 {{ $config.Serilog | toYaml | nindent 2 }}
 {{- end }}
 
+{{- if $config.CacheRef }}
+CacheRef:
+{{ $config.CacheRef | toYaml | nindent 2 }}
+{{- end }}
+
 {{- if and (.Values.global.cloudProvider) (eq .Values.global.cloudProvider "Azure") }}
 {{ include "cloud.settings.azure" .Values }}
 {{- else if and (.Values.global.cloudProvider) (eq .Values.global.cloudProvider "AWS") }}
diff --git a/Engine/Source/Programs/UnrealCloudDDC/Jupiter/Jupiter.Startup.cs b/Engine/Source/Programs/UnrealCloudDDC/Jupiter/Jupiter.Startup.cs
--- a/Engine/Source/Programs/UnrealCloudDDC/Jupiter/Jupiter.Startup.cs
+++ b/Engine/Source/Programs/UnrealCloudDDC/Jupiter/Jupiter.Startup.cs
@@ -106,6 +106,8 @@ namespace Jupiter
             services.AddOptions<KubernetesLeaderElectionSettings>().Configure(o => Configuration.GetSection("Kubernetes").Bind(o)).ValidateDataAnnotations();
             services.AddOptions<StaticPeerServiceDiscoverySettings>().Configure(o => Configuration.GetSection("ServiceDiscovery").Bind(o)).ValidateDataAnnotations();
 
+            services.AddOptions<MemoryCacheReferencesSettings>().Configure(o => Configuration.GetSection("CacheRef").Bind(o)).ValidateDataAnnotations();
+
             services.AddSingleton(typeof(CompressedBufferUtils), CreateCompressedBufferUtils);
 
             services.AddSingleton<AWSCredentials>(provider =>
@@ -417,7 +419,7 @@ namespace Jupiter
                     $"Unknown references db implementation: {settings.ReferencesDbImplementation}")
             };
 
-            MemoryCacheContentIdSettings memoryCacheSettings = provider.GetService<IOptionsMonitor<MemoryCacheContentIdSettings>>()!.CurrentValue;
+            MemoryCacheReferencesSettings memoryCacheSettings = provider.GetService<IOptionsMonitor<MemoryCacheReferencesSettings>>()!.CurrentValue;
 
             if (memoryCacheSettings.Enabled)
             {
