Unreal Cloud DDC - Use recommended way to create local DC load balancing policy for scylla (this makes sure the load balancing policy is token aware)

diff --git a/Engine/Source/Programs/UnrealCloudDDC/Jupiter/Jupiter.Startup.cs b/Engine/Source/Programs/UnrealCloudDDC/Jupiter/Jupiter.Startup.cs
--- a/Engine/Source/Programs/UnrealCloudDDC/Jupiter/Jupiter.Startup.cs
+++ b/Engine/Source/Programs/UnrealCloudDDC/Jupiter/Jupiter.Startup.cs
@@ -293,7 +293,7 @@ namespace Jupiter
             Builder clusterBuilder = Cluster.Builder()
                 .WithConnectionString(connectionString)
                 .WithDefaultKeyspace(DefaultKeyspaceName)
-                .WithLoadBalancingPolicy(new DefaultLoadBalancingPolicy(settings.LocalDatacenterName))
+                .WithLoadBalancingPolicy(Policies.NewDefaultLoadBalancingPolicy(settings.LocalDatacenterName))
                 .WithPoolingOptions(PoolingOptions.Create().SetMaxConnectionsPerHost(HostDistance.Local, settings.MaxConnectionForLocalHost))
                 .WithExecutionProfiles(options =>
                     options.WithProfile("default", builder => builder.WithConsistencyLevel(ConsistencyLevel.LocalOne)));
