Unreal Cloud DDC - Updated changelog with the latest changes to be rolled out

diff --git a/Engine/Source/Programs/UnrealCloudDDC/changelog.md b/Engine/Source/Programs/UnrealCloudDDC/changelog.md
--- a/Engine/Source/Programs/UnrealCloudDDC/changelog.md
+++ b/Engine/Source/Programs/UnrealCloudDDC/changelog.md
@@ -1,6 +1,10 @@
 # Unreleased
 * Azure blob storage now supports storage pools
-* Blobs will now be GCed when a ref is GCed (orphan blob GC can and should still be used for partial uploads)
+* Last access table refactoring - Moved the last accessing tracking out of the objects table and into a seperate table. Saves on compation work for Scylla. Set `Scylla.ListObjectsFromLastAccessTable` to migrate GC refs to use this new table (will be default in the next release).
+* Blob index table refactoring - Multiple changes to the blob index table, avoiding collections in the scylla rows and instead prefering clustering keys. This reduces amount of work that needs to happen when updating regions or references to blobs and improves performance (and caching) on the DB. We automatically migrate old blob index entiries when `Scylla.MigrateFromOldBlobIndex` is set (it is by default). We will start requring data in these tables by the next release. 
+* Generally reduced some of the excessive logging from the worker in favor of open telemetry metrics instead.
+* GC Rules - Added ability to disable GC or configure it to use Time-To-Live instead of last access tracking.
+* FallbackNamespaces added - allows you to configure a second namespace used for blob operations if blobs are missing from the first. This can be used to store Virtual Assets in a seperate store pool that have different GC rules.
 
 # 0.2.0
 * Improved handling of cassandra timeouts during ref GC.
