Unreal Cloud DDC - Migrated to new namespace / bucket table layout - we will read both the old and new table by default while we migrate. Also avoiding writing to scylla if bucket has already been added (caching this on write) to reduce unnesscary traffic.

diff --git a/Engine/Source/Programs/UnrealCloudDDC/Jupiter/Implementation/Objects/ScyllaReferencesStore.cs b/Engine/Source/Programs/UnrealCloudDDC/Jupiter/Implementation/Objects/ScyllaReferencesStore.cs
--- a/Engine/Source/Programs/UnrealCloudDDC/Jupiter/Implementation/Objects/ScyllaReferencesStore.cs
+++ b/Engine/Source/Programs/UnrealCloudDDC/Jupiter/Implementation/Objects/ScyllaReferencesStore.cs
@@ -1,6 +1,7 @@
 // Copyright Epic Games, Inc. All Rights Reserved.
 
 using System;
+using System.Collections.Concurrent;
 using System.Collections.Generic;
 using System.Linq;
 using System.Threading.Tasks;
@@ -25,8 +26,10 @@ namespace Jupiter.Implementation
         private readonly PreparedStatement _getObjectsStatement;
         private readonly PreparedStatement _getObjectsLastAccessStatement;
         private readonly PreparedStatement _getNamespacesStatement;
+        private readonly PreparedStatement _getNamespacesOldStatement;
         private readonly PreparedStatement _getObjectsForPartitionRangeStatement;
         private readonly PreparedStatement _getObjectsLastAccessForPartitionRangeStatement;
+        private readonly ConcurrentDictionary<NamespaceId, ConcurrentBag<BucketId>> _addedBuckets = new ConcurrentDictionary<NamespaceId, ConcurrentBag<BucketId>>();
 
         public ScyllaReferencesStore(IScyllaSessionManager scyllaSessionManager, IOptionsMonitor<ScyllaSettings> settings, INamespacePolicyResolver namespacePolicyResolver, Tracer tracer, ILogger<ScyllaReferencesStore> logger)
         {
@@ -59,10 +62,10 @@ namespace Jupiter.Implementation
             );"
             ));
             
-            _session.Execute(new SimpleStatement(@"CREATE TABLE IF NOT EXISTS buckets (
+            _session.Execute(new SimpleStatement(@"CREATE TABLE IF NOT EXISTS buckets_v2 (
                 namespace text, 
-                bucket set<text>, 
-                PRIMARY KEY (namespace)
+                bucket text, 
+                PRIMARY KEY ((namespace), bucket)
             );"
             ));
 
@@ -70,7 +73,8 @@ namespace Jupiter.Implementation
             string cqlOptions = scyllaSessionManager.IsScylla ? "BYPASS CACHE" : "";
             _getObjectsStatement = _session.Prepare($"SELECT namespace, bucket, name, last_access_time FROM objects ALLOW FILTERING {cqlOptions}");
             _getObjectsLastAccessStatement = _session.Prepare($"SELECT namespace, bucket, name, last_access_time FROM object_last_access_v2 ALLOW FILTERING {cqlOptions}");
-            _getNamespacesStatement = _session.Prepare("SELECT DISTINCT namespace FROM buckets");
+            _getNamespacesStatement = _session.Prepare("SELECT DISTINCT namespace FROM buckets_v2");
+            _getNamespacesOldStatement = _session.Prepare("SELECT DISTINCT namespace FROM buckets");
 
             _getObjectsForPartitionRangeStatement = _session.Prepare($"SELECT namespace, bucket, name, last_access_time FROM objects WHERE token(namespace, bucket, name) >= ? AND token(namespace, bucket, name) <= ? {cqlOptions}");
             _getObjectsLastAccessForPartitionRangeStatement = _session.Prepare($"SELECT namespace, bucket, name, last_access_time FROM object_last_access_v2 WHERE token(namespace, bucket, name) >= ? AND token(namespace, bucket, name) <= ? {cqlOptions}");
@@ -121,7 +125,7 @@ namespace Jupiter.Implementation
             }
 
             // add the bucket in parallel with inserting the actual object
-            Task addBucketTask = MaybeAddBucket(ns, bucket);
+            Task addBucketTask = AddBucket(ns, bucket);
 
             int? ttl = null;
             NamespacePolicy policy = _namespacePolicyResolver.GetPoliciesForNs(ns);
@@ -280,16 +284,37 @@ namespace Jupiter.Implementation
         public async IAsyncEnumerable<NamespaceId> GetNamespaces()
         {
             using TelemetrySpan scope = _tracer.BuildScyllaSpan("scylla.get_namespaces");
-            RowSet rowSet = await _session.ExecuteAsync(_getNamespacesStatement.Bind());
 
-            foreach (Row row in rowSet)
             {
-                if (rowSet.GetAvailableWithoutFetching() == 0)
+                RowSet rowSet = await _session.ExecuteAsync(_getNamespacesStatement.Bind());
+
+                foreach (Row row in rowSet)
                 {
-                    await rowSet.FetchMoreResultsAsync();
+                    if (rowSet.GetAvailableWithoutFetching() == 0)
+                    {
+                        await rowSet.FetchMoreResultsAsync();
+                    }
+
+                    yield return new NamespaceId(row.GetValue<string>(0));
                 }
+            }
+
+            if (_settings.CurrentValue.ListObjectsFromOldNamespaceTable)
+            {
+                // this will likely generate duplicates from the statements above but that is not a huge issue
+                using TelemetrySpan _ = _tracer.BuildScyllaSpan("scylla.get_old_namespaces");
+
+                RowSet rowSet = await _session.ExecuteAsync(_getNamespacesOldStatement.Bind());
 
-                yield return new NamespaceId(row.GetValue<string>(0));
+                foreach (Row row in rowSet)
+                {
+                    if (rowSet.GetAvailableWithoutFetching() == 0)
+                    {
+                        await rowSet.FetchMoreResultsAsync();
+                    }
+
+                    yield return new NamespaceId(row.GetValue<string>(0));
+                }
             }
         }
 
@@ -331,6 +356,7 @@ namespace Jupiter.Implementation
         public async Task<long> DeleteBucket(NamespaceId ns, BucketId bucket)
         {
             using TelemetrySpan scope = _tracer.BuildScyllaSpan("scylla.delete_bucket");
+
             RowSet rowSet = await _session.ExecuteAsync(new SimpleStatement("SELECT name FROM objects WHERE namespace = ? AND bucket = ? ALLOW FILTERING;", ns.ToString(), bucket.ToString()));
             long deletedCount = 0;
             foreach (Row row in rowSet)
@@ -342,15 +368,24 @@ namespace Jupiter.Implementation
             }
 
             // remove the tracking in the buckets table as well
-            await _mapper.UpdateAsync<ScyllaBucket>("SET bucket = bucket - ? WHERE namespace = ?", new string[] {bucket.ToString()}, ns.ToString());
+            await _mapper.DeleteAsync<ScyllaBucket>(new ScyllaBucket(ns, bucket));
 
             return deletedCount;
         }
 
-        private async Task MaybeAddBucket(NamespaceId ns, BucketId bucket)
+        private async Task AddBucket(NamespaceId ns, BucketId bucket)
         {
             using TelemetrySpan scope = _tracer.BuildScyllaSpan("scylla.add_bucket");
-            await _mapper.UpdateAsync<ScyllaBucket>("SET bucket = bucket + ? WHERE namespace = ?", new string[] {bucket.ToString()}, ns.ToString());
+
+            ConcurrentBag<BucketId> addedBuckets = _addedBuckets.GetOrAdd(ns, id => new ConcurrentBag<BucketId>());
+
+            bool alreadyAdded = addedBuckets.Contains(bucket);
+            if (!alreadyAdded)
+            {
+                Task addTask = _mapper.InsertAsync<ScyllaBucket>(new ScyllaBucket(ns, bucket));
+                addedBuckets.Add(bucket);
+                await addTask;
+            }
         }
     }
 
@@ -475,25 +510,20 @@ namespace Jupiter.Implementation
         }
     }
 
-    [Cassandra.Mapping.Attributes.Table("buckets")]
+    [Cassandra.Mapping.Attributes.Table("buckets_v2")]
     public class ScyllaBucket
     {
-        public ScyllaBucket()
-        {
-
-        }
-
-        public ScyllaBucket(NamespaceId ns, BucketId[] buckets)
+        public ScyllaBucket(NamespaceId ns, BucketId bucket)
         {
             Namespace = ns.ToString();
-            Buckets = buckets.Select(b => b.ToString()).ToList();
+            Bucket = bucket.ToString();
         }
 
         [Cassandra.Mapping.Attributes.PartitionKey]
         public string? Namespace { get; set; }
 
-        [System.Diagnostics.CodeAnalysis.SuppressMessage("Usage", "CA2227:Collection properties should be read only", Justification = "Used by serialization")]
-        public List<string> Buckets { get; set; } = new List<string>();
+        [Cassandra.Mapping.Attributes.ClusteringKey]
+        public string Bucket { get; set; }
     }
 
     
diff --git a/Engine/Source/Programs/UnrealCloudDDC/Jupiter/Implementation/Objects/ScyllaSettings.cs b/Engine/Source/Programs/UnrealCloudDDC/Jupiter/Implementation/Objects/ScyllaSettings.cs
--- a/Engine/Source/Programs/UnrealCloudDDC/Jupiter/Implementation/Objects/ScyllaSettings.cs
+++ b/Engine/Source/Programs/UnrealCloudDDC/Jupiter/Implementation/Objects/ScyllaSettings.cs
@@ -98,6 +98,11 @@ public class ScyllaSettings : IValidatableObject
     /// </summary>
     public bool MigrateFromOldBlobIndex { get; set; } = true;
 
+    /// <summary>
+    /// Enable to use list namespaces from the old table, this can set for a period while migration from old to new table happens
+    /// </summary>
+    public bool ListObjectsFromOldNamespaceTable { get; set; } = true;
+
     public IEnumerable<ValidationResult> Validate(ValidationContext validationContext)
     {
         List<ValidationResult> results = new List<ValidationResult>();
diff --git a/Engine/Source/Programs/UnrealCloudDDC/changelog.md b/Engine/Source/Programs/UnrealCloudDDC/changelog.md
--- a/Engine/Source/Programs/UnrealCloudDDC/changelog.md
+++ b/Engine/Source/Programs/UnrealCloudDDC/changelog.md
@@ -2,6 +2,7 @@
 * Azure blob storage now supports storage pools
 * Last access table refactoring - Moved the last accessing tracking out of the objects table and into a seperate table. Saves on compation work for Scylla. Set `Scylla.ListObjectsFromLastAccessTable` to migrate GC refs to use this new table (will be default in the next release).
 * Blob index table refactoring - Multiple changes to the blob index table, avoiding collections in the scylla rows and instead prefering clustering keys. This reduces amount of work that needs to happen when updating regions or references to blobs and improves performance (and caching) on the DB. We automatically migrate old blob index entiries when `Scylla.MigrateFromOldBlobIndex` is set (it is by default). We will start requring data in these tables by the next release. 
+* Bucket table refactoring - Reads data from old and new table for this release. Set `Scylla.ListObjectsFromOldNamespaceTable` to disable reading the old table.
 * Generally reduced some of the excessive logging from the worker in favor of open telemetry metrics instead.
 * GC Rules - Added ability to disable GC or configure it to use Time-To-Live instead of last access tracking.
 * FallbackNamespaces added - allows you to configure a second namespace used for blob operations if blobs are missing from the first. This can be used to store Virtual Assets in a seperate store pool that have different GC rules.
