Unreal Cloud DDC - Updated benchmarks

* Added benchmark using compact binary packages
* Avoid storing bodies from responses, allowing the large file tests to run (re-enabled them)
* Added game tests which generates a bunch of random assets with a distribution based from actual game data and tests upload and download
* Added --seed-remote option which is required to actually seed the remote instance , this can be used to avoid seeding the server in between different tests

diff --git a/Engine/Source/Programs/UnrealCloudDDC/Benchmarks/benchmark-suite/helpers.py b/Engine/Source/Programs/UnrealCloudDDC/Benchmarks/benchmark-suite/helpers.py
--- a/Engine/Source/Programs/UnrealCloudDDC/Benchmarks/benchmark-suite/helpers.py
+++ b/Engine/Source/Programs/UnrealCloudDDC/Benchmarks/benchmark-suite/helpers.py
@@ -8,6 +8,7 @@ import random
 import requests
 from blake3 import blake3
 import varint
+import scipy
 
 logger = logging.getLogger("helpers")
 
@@ -43,6 +44,51 @@ def generate_blob(args, seed, file_length):
 
     return write_payload(args, content)
 
+# generates a file with random content were the size follows distributions from how large the cache of a game can be
+def generate_game_asset(args, seed):
+    def random_game_asset_size():
+        value = random.random()
+        if (value < 0.25): # 25% chance
+            return 130
+        if (value < 0.31): # 6% chance
+            return 386
+        if (value < 0.34): # 3% chance
+            return 5136  
+        if (value < 0.37): # 3% chance
+            return 65608
+        if (value < 0.41): # 4% chance
+            return 350 # approximation for medium set
+            
+        if (value > 0.999): # 0.1 chance
+            # our precision on really large assets is missing so making sure we have some of those
+            # approximating a avg of the large assets
+            return 84_186_659
+
+        if (value > 0.99): # 1% chance
+            # really large values have a special log distribution
+            s_large = 1.1272032578046247,
+            loc_large = 233474.20499173328,
+            scale_large = 292597.08600436774
+            return int(scipy.stats.lognorm(s=s_large, loc=loc_large, scale=scale_large).rvs(size=1, random_state=seed)[0])
+        
+        # follow lognormal distribution for the rest    
+        s = 1.7898978392258396,
+        loc = 489.70804173566125,
+        scale = 13112.784274036905
+
+        return int(scipy.stats.lognorm(s=s, loc=loc, scale=scale).rvs(size=1, random_state=seed)[0])
+
+    random.seed(seed)
+    
+    file_length = random_game_asset_size()
+    # generate a byte array of file_length,
+    # hash it to determine its identifier and then write it to disk
+
+    # randbytes is python 3.9 only, so we steal the implementation of randbytes and use that below to be compatible with older versions of Python
+    #content = random.randbytes(file_length)
+    content = random.getrandbits(file_length * 8).to_bytes(file_length, 'little')
+
+    return write_payload(args, content)
 
 def generate_uecb_singlefield(args, fieldName, attachments):
     # object with a single field
@@ -113,8 +159,10 @@ def run_test(args, test_name, duration_seconds=5, rate_per_second=None):
     results_file = os.path.join(args['reports_dir'], test_name)
 
     run_vegeta(["attack",
+                f"-name={test_name}",
                 f"-duration={duration_seconds}s" if duration_seconds is not None else '',
                 f"-rate={rate_per_second}/s" if rate_per_second is not None else '',
+                "-max-body=0", # do not capture the response body as it becomes very large
                 "-workers=64", # bump the number of workers and thus the number of connections we start with
                 f"-targets={test_file}",
                 f"-output={results_file}"])
diff --git a/Engine/Source/Programs/UnrealCloudDDC/Benchmarks/benchmark-suite/main.py b/Engine/Source/Programs/UnrealCloudDDC/Benchmarks/benchmark-suite/main.py
--- a/Engine/Source/Programs/UnrealCloudDDC/Benchmarks/benchmark-suite/main.py
+++ b/Engine/Source/Programs/UnrealCloudDDC/Benchmarks/benchmark-suite/main.py
@@ -53,6 +53,8 @@ def main():
         description="UnrealCloudDDC Benchmarker - Tool for benchmarking UnrealCloudDDC")
     parser.add_argument('--seed', action="store_true",
                         help='Set to generate test data')
+    parser.add_argument('--seed-remote', action="store_true",
+                        help='Set to generate test data')
     parser.add_argument('--list', action="store_true", help='List known tests')
     # host.docker.internal is the ip of the hostmachine for docker
     parser.add_argument(
@@ -83,7 +85,8 @@ def main():
         'payloads_dir': os.path.join(".", "payloads"),
         'reports_dir': os.path.join(".", "reports"),
         'headers': [] if args.header is None else args.header,
-        'namespace': args.namespace
+        'namespace': args.namespace,
+        'seed-remote': args.seed_remote
     }
 
     for name in ['tests_dir', 'payloads_dir', 'reports_dir']:
diff --git a/Engine/Source/Programs/UnrealCloudDDC/Benchmarks/benchmark-suite/requirements.txt b/Engine/Source/Programs/UnrealCloudDDC/Benchmarks/benchmark-suite/requirements.txt
--- a/Engine/Source/Programs/UnrealCloudDDC/Benchmarks/benchmark-suite/requirements.txt
+++ b/Engine/Source/Programs/UnrealCloudDDC/Benchmarks/benchmark-suite/requirements.txt
@@ -1,4 +1,5 @@
 blake3
 varint
 requests
-json_log_formatter
\ No newline at end of file
+json_log_formatter
+scipy
\ No newline at end of file
diff --git a/Engine/Source/Programs/UnrealCloudDDC/Benchmarks/benchmark-suite/tests.py b/Engine/Source/Programs/UnrealCloudDDC/Benchmarks/benchmark-suite/tests.py
--- a/Engine/Source/Programs/UnrealCloudDDC/Benchmarks/benchmark-suite/tests.py
+++ b/Engine/Source/Programs/UnrealCloudDDC/Benchmarks/benchmark-suite/tests.py
@@ -9,7 +9,8 @@ start_seeds = {
     'small_blob_seed': 1001,
     'large_blob_seed': 1501,
     'refs_raw_seed': 2001,
-    'compact_binary_seed': 3001
+    'compact_binary_seed': 3001,
+    'game_seed': 5001
 }
 
 
@@ -31,17 +32,18 @@ def test_healthcheck_live(args):
 
 def generate_and_upload_blobs(args, start_seed, count, file_length):
     ns = args['namespace']
-
+    seed_remote = args['seed-remote']
     blobs = [helpers.generate_blob(args, start_seed + i, file_length)
              for i in range(0, count)]
 
-    session = requests.session()
-    # upload the refs so that they are present for downloads
-    logger.info(f"Seeding {count} blobs to server with size {file_length / 1024} kb")
-    for id, payload_path in blobs:
-        id, content = helpers.get_content(payload_path)
-        helpers.execute_http_call(session, args, 'PUT', f"/api/v1/blobs/{ns}/{id}", body=content, headers=[
-                                  "Content-Type: application/octet-stream"])
+    if seed_remote:
+        session = requests.session()
+        # upload the refs so that they are present for downloads
+        logger.info(f"Seeding {count} blobs to server with size {file_length / 1024} kb")
+        for id, payload_path in blobs:
+            id, content = helpers.get_content(payload_path)
+            helpers.execute_http_call(session, args, 'PUT', f"/api/v1/blobs/{ns}/{id}", body=content, headers=[
+                                    "Content-Type: application/octet-stream"])
     return blobs
 
 def seed_small_blobs(args):
@@ -84,18 +86,16 @@ def seed_large_blobs(args):
         args, 'GET', f"/api/v1/blobs/{ns}/{id}") for id, body in blobs]
     helpers.write_test('large_blob_downloads', test_contents_download, args)
 
-# Do not run the large blob tests as Vegeta does not handle them very well, it will crash likely due to having to write all of these huge bodies to its results file
-#def test_large_blob_uploads(args):
-#    helpers.run_test(args, 'large_blob_uploads', duration_seconds=60, rate_per_second=10)
-
-
-#def test_large_blob_downloads(args):
-#    helpers.run_test(args, 'large_blob_downloads', duration_seconds=60, rate_per_second=10)
+def test_large_blob_uploads(args):
+    helpers.run_test(args, 'large_blob_uploads', duration_seconds=60, rate_per_second=10)
 
+def test_large_blob_downloads(args):
+    helpers.run_test(args, 'large_blob_downloads', duration_seconds=60, rate_per_second=10)
 
 def seed_refs(args):
     ns = args['namespace']
     bucket = 'default'
+    seed_remote = args['seed-remote']
 
     seed = start_seeds['refs_raw_seed']
     file_length = 2 * 1024 * 1024  # 2 MB refs
@@ -106,13 +106,14 @@ def seed_refs(args):
         helpers.hash_content(b'ref_' + id.encode()).hex(),
         body) for id, body in blobs]
 
-    session = requests.session()
-    # upload the refs so that they are present for downloads
-    logger.info("Seeding refs to server")
-    for ref_id, payload_path in refs:
-        id, content = helpers.get_content(payload_path)
-        helpers.execute_http_call(session, args, 'PUT', f"/api/v1/refs/{ns}/{bucket}/{ref_id}", body=content, headers=[
-                                  "Content-Type: application/octet-stream", f"X-Jupiter-IoHash: {id}"])
+    if seed_remote:
+        session = requests.session()
+        # upload the refs so that they are present for downloads
+        logger.info("Seeding refs to server")
+        for ref_id, payload_path in refs:
+            id, content = helpers.get_content(payload_path)
+            helpers.execute_http_call(session, args, 'PUT', f"/api/v1/refs/{ns}/{bucket}/{ref_id}", body=content, headers=[
+                                    "Content-Type: application/octet-stream", f"X-Jupiter-IoHash: {id}"])
 
     # we generate a ref hash based on the content id and a random string, the ref identifier is not intresting but needs to be stable
     test_contents_download = [helpers.build_http_call(
@@ -133,6 +134,7 @@ def test_refs_raw_upload(args):
 
 
 def seed_compact_binary(args):
+    seed_remote = args['seed-remote']
     ns = args['namespace']
     bucket = 'default'
 
@@ -150,26 +152,31 @@ def seed_compact_binary(args):
     for ref_id, cb, attachment_path in compact_binaries:
         cb_payloads[ref_id] = helpers.write_payload(args, cb)
 
-    session = requests.session()
-    logger.info("Seeding compact binary attachments to server")
-    # upload the blobs we will use as attachments
-    for ref_id, cb, attachment_path in compact_binaries:
-        id, content = helpers.get_content(attachment_path)
-        helpers.execute_http_call(session, args, 'PUT', f"/api/v1/blobs/{ns}/{id}", body=content, headers=[
-                                  "Content-Type: application/octet-stream"])
-
-    logger.info("Seeding compact binary refs to server")
-    # make sure the refs we which to download are present on the server
-    for ref_id, cb, attachment_path in compact_binaries:
-        hash = helpers.hash_content(cb).hex()
-        helpers.execute_http_call(session, args, 'PUT', f"/api/v1/refs/{ns}/{bucket}/{ref_id}", body=cb, headers=[
-                                  "Content-Type: application/x-ue-cb", f"X-Jupiter-IoHash: {hash}"])
+    if seed_remote:
+        session = requests.session()
+        logger.info("Seeding compact binary attachments to server")
+        # upload the blobs we will use as attachments
+        for ref_id, cb, attachment_path in compact_binaries:
+            id, content = helpers.get_content(attachment_path)
+            helpers.execute_http_call(session, args, 'PUT', f"/api/v1/blobs/{ns}/{id}", body=content, headers=[
+                                    "Content-Type: application/octet-stream"])
+
+        logger.info("Seeding compact binary refs to server")
+        # make sure the refs we which to download are present on the server
+        for ref_id, cb, attachment_path in compact_binaries:
+            hash = helpers.hash_content(cb).hex()
+            helpers.execute_http_call(session, args, 'PUT', f"/api/v1/refs/{ns}/{bucket}/{ref_id}", body=cb, headers=[
+                                    "Content-Type: application/x-ue-cb", f"X-Jupiter-IoHash: {hash}"])
 
     # we generate a ref hash based on the content id and a random string, the ref identifier is not intresting but needs to be stable
     test_contents_download = [helpers.build_http_call(
         args, 'GET', f"/api/v1/refs/{ns}/{bucket}/{ref_id}.uecb") for ref_id, cb, attachment_path in compact_binaries]
     helpers.write_test('uecb_download', test_contents_download, args)
 
+    test_contents_package_download = [helpers.build_http_call(
+        args, 'GET', f"/api/v1/refs/{ns}/{bucket}/{ref_id}.uecbpkg") for ref_id, cb, attachment_path in compact_binaries]
+    helpers.write_test('uecb_pkg_download', test_contents_package_download, args)
+
     test_contents_upload = [helpers.build_http_call(args, 'PUT', f"/api/v1/refs/{ns}/{bucket}/{ref_id}", body=cb_payloads[ref_id][1], headers=[
                                                     "Content-Type: application/x-ue-cb", f"X-Jupiter-IoHash: {cb_payloads[ref_id][0]}"]) for ref_id, cb, attachment_path in compact_binaries]
     helpers.write_test('uecb_upload', test_contents_upload, args)
@@ -178,6 +185,60 @@ def seed_compact_binary(args):
 def test_uecb_download(args):
     helpers.run_test(args, 'uecb_download', duration_seconds=15)
 
+def test_uecb_pkg_download(args):
+    helpers.run_test(args, 'uecb_pkg_download', duration_seconds=15)
 
 def test_uecb_upload(args):
     helpers.run_test(args, 'uecb_upload', duration_seconds=15)
+
+
+def test_game_upload(args):
+    helpers.run_test(args, 'game_upload', duration_seconds=120)
+
+def test_game_download(args):
+    helpers.run_test(args, 'game_download', duration_seconds=120)
+
+def seed_game(args):
+    seed = start_seeds['game_seed']
+    seed_remote = args['seed-remote']
+    ns = args['namespace']
+    bucket = 'default'
+
+    logger.info("Generating game blobs")
+    blobs = [helpers.generate_game_asset(args, seed + i) for i in range(0, 1000)]
+
+    compact_binaries = [(
+        helpers.hash_content(b'ref_' + id.encode()).hex(),
+        helpers.generate_uecb_singlefield(
+            args, 'payload', attachments=[bytes.fromhex(id)]),
+        body) for id, body in blobs
+    ]
+
+    cb_payloads = {}
+    for ref_id, cb, attachment_path in compact_binaries:
+        cb_payloads[ref_id] = helpers.write_payload(args, cb)
+    if seed_remote:
+        session = requests.session()
+    
+        logger.info("Seeding game blobs to server")
+        # upload the blobs we will use as attachments
+        for ref_id, cb, attachment_path in compact_binaries:
+            id, content = helpers.get_content(attachment_path)
+            helpers.execute_http_call(session, args, 'PUT', f"/api/v1/blobs/{ns}/{id}", body=content, headers=[
+                                    "Content-Type: application/octet-stream"])
+
+        logger.info("Seeding game refs to server")
+        # make sure the refs we which to download are present on the server
+        for ref_id, cb, attachment_path in compact_binaries:
+            hash = helpers.hash_content(cb).hex()
+            helpers.execute_http_call(session, args, 'PUT', f"/api/v1/refs/{ns}/{bucket}/{ref_id}", body=cb, headers=[
+                                    "Content-Type: application/x-ue-cb", f"X-Jupiter-IoHash: {hash}"])
+
+    # we generate a ref hash based on the content id and a random string, the ref identifier is not intresting but needs to be stable
+    test_contents_package_download = [helpers.build_http_call(
+        args, 'GET', f"/api/v1/refs/{ns}/{bucket}/{ref_id}.uecbpkg") for ref_id, cb, attachment_path in compact_binaries]
+    helpers.write_test('game_download', test_contents_package_download, args)
+
+    test_contents_upload = [helpers.build_http_call(args, 'PUT', f"/api/v1/refs/{ns}/{bucket}/{ref_id}", body=cb_payloads[ref_id][1], headers=[
+                                                    "Content-Type: application/x-ue-cb", f"X-Jupiter-IoHash: {cb_payloads[ref_id][0]}"]) for ref_id, cb, attachment_path in compact_binaries]
+    helpers.write_test('game_upload', test_contents_upload, args)
\ No newline at end of file
diff --git a/Engine/Source/Programs/UnrealCloudDDC/Benchmarks/readme.md b/Engine/Source/Programs/UnrealCloudDDC/Benchmarks/readme.md
--- a/Engine/Source/Programs/UnrealCloudDDC/Benchmarks/readme.md
+++ b/Engine/Source/Programs/UnrealCloudDDC/Benchmarks/readme.md
@@ -1,15 +1,39 @@
+# Intro
 Under `benchmark-suite` you will find a python program and dockerfile to build it into a container that can run benchmarks against UnrealCloudDDC.
+
+# Usage
 We use vegeta to run these benchmarks and the python app generates test data which vegeta can use to run concurrent http tests.
 
 The docker image will automatically run all tests on startup, so we usually use this to just start the docker image somewhere and having that run its tests.
 The image expects to get a valid access token passed to it as a argument.
 Usage:
 ```
-docker build -t jupiter_benchmark
+docker build -t jupiter_benchmark .
 docker run --network host jupiter_benchmark \
 --seed --host <url-to-service> \
 --header="Authorization: Bearer <access-token>" \
 all
 ```
 
-While the python script can be executed on any platform vegeta only works for linux, if you have a linux machine you can run these locally otherwise run them via docker.
\ No newline at end of file
+While the python script can be executed on any platform vegeta only works for linux, if you have a linux machine you can run these locally otherwise run them via docker.
+
+# Expected results
+
+These tests assumes they are running on a machine with good network and a lot of cores to generate enough load on your UnrealCloudDDC deployment, we typically run it with 100 gigabit networking, in the same datacenter as the deployment (to avoid network conditions) and 64 cores which is what is assumed for these expected results.
+
+If you do want to factor in the networking conditions you can run this elsewhere but those results can not compared to these.
+
+This a rough guide for what kind of results you should be seeing if your setup is correct
+
+| Test                 | p50      | p95     |
+|----------------------|----------|---------|
+| healthcheck_ready    | 1.3 ms   | 1.8 ms  |
+| uecb_download        | 1.5 ms   | 2.0 ms  |
+| small_blob_uploads   | 6.6 ms   | 13.0 ms |
+| small_blob_downloads | 2.1 ms   | 3 ms    |
+| large_blob_uploads   | 470 ms   | 880 ms  |
+| large_blob_download  | 34.4 ms  | 38.8 ms |
+| refs_raw_download    | 5.7 ms   | 7 ms    |
+| uecb_pkg_download    | 4.9 ms   | 5.6 ms  |
+| game_download        | 3.4 ms   | 129 ms  |
+| game_upload          | 42.2 ms  | 88 ms   |
\ No newline at end of file
