Unreal Cloud DDC - Added support for running nginx as a reverse proxy together with our container in the helm chart. This is also configured to use sendfile to send files from disk if available.

diff --git a/Engine/Source/Programs/UnrealCloudDDC/Helm/UnrealCloudDDC/templates/configmap.yaml b/Engine/Source/Programs/UnrealCloudDDC/Helm/UnrealCloudDDC/templates/configmap.yaml
--- a/Engine/Source/Programs/UnrealCloudDDC/Helm/UnrealCloudDDC/templates/configmap.yaml
+++ b/Engine/Source/Programs/UnrealCloudDDC/Helm/UnrealCloudDDC/templates/configmap.yaml
@@ -64,6 +64,11 @@ Jupiter:
 {{ $config.Jupiter | toYaml | nindent 2 }}
 {{- end }}
 
+{{ if .Values.nginx.enabled }}
+Nginx:
+  UseNginxRedirect: true
+{{ end }}
+
 Replication:
   Enabled: {{ default "false" $replication.Enabled }}
 {{ with $replication.Replicators }}
diff --git a/Engine/Source/Programs/UnrealCloudDDC/Helm/UnrealCloudDDC/templates/deployment.yaml b/Engine/Source/Programs/UnrealCloudDDC/Helm/UnrealCloudDDC/templates/deployment.yaml
--- a/Engine/Source/Programs/UnrealCloudDDC/Helm/UnrealCloudDDC/templates/deployment.yaml
+++ b/Engine/Source/Programs/UnrealCloudDDC/Helm/UnrealCloudDDC/templates/deployment.yaml
@@ -101,6 +101,10 @@ spec:
         - name: diagnostics
           mountPath: /tmp
         {{- end }}
+    {{- if .Values.nginx.enabled }}
+        - name: nginx
+          mountPath: /nginx
+    {{- end }}
         {{- if or .Values.persistence.volume .Values.persistence.enabled }}
         - name: data
           mountPath: {{ .Values.persistence.mountPath }}
@@ -115,6 +119,32 @@ spec:
         volumeMounts:
           - name: diagnostics
             mountPath: /tmp
+    {{- end }}
+    {{- if .Values.nginx.enabled }}
+      - name: nginx
+        image: nginx:1.23.3
+        {{- if .Values.nginx.env }}
+        env:
+        {{- include "common.tplvalues.render" ( dict "value" .Values.nginx.env "context" $ ) | nindent 8 }}
+        {{- end }}
+        ports:
+        - name: nginx-http
+          containerPort: 81
+        - name: nginx-corp-http
+          containerPort: 8009
+        - name: nginx-h2
+          containerPort: 8091
+        - name: nginx-corp-h
+          containerPort: 8092
+        volumeMounts:
+          - name: nginx
+            mountPath: /nginx
+          - name: nginx-config-volume
+            mountPath: /etc/nginx
+        {{- if or .Values.persistence.volume .Values.persistence.enabled }}
+          - name: data
+            mountPath: {{ .Values.persistence.mountPath }}
+        {{- end }}
     {{- end }}
       volumes:
       - name: config-volume
@@ -123,6 +153,13 @@ spec:
     {{- if .Values.useDotnetMonitorSidecar }}
       - name: diagnostics
         emptyDir: {}
+    {{- end }}
+    {{- if .Values.nginx.enabled }}
+      - name: nginx
+        emptyDir: {}
+      - name: nginx-config-volume
+        configMap:
+          name: {{ $fullname }}-nginx
     {{- end }}
       {{- if or .Values.persistence.enabled .Values.persistence.volume }}
       - name: data
diff --git a/Engine/Source/Programs/UnrealCloudDDC/Helm/UnrealCloudDDC/templates/extra-service-2.yaml b/Engine/Source/Programs/UnrealCloudDDC/Helm/UnrealCloudDDC/templates/extra-service-2.yaml
--- a/Engine/Source/Programs/UnrealCloudDDC/Helm/UnrealCloudDDC/templates/extra-service-2.yaml
+++ b/Engine/Source/Programs/UnrealCloudDDC/Helm/UnrealCloudDDC/templates/extra-service-2.yaml
@@ -17,5 +17,8 @@ spec:
     - name: {{ .Values.extraService2.portName }}
       port: {{ .Values.extraService2.port }}
       targetPort: {{ default "http" .Values.extraService2.targetPort }}
+  {{- if .Values.extraService2.extraPort }}
+  {{- toYaml .Values.extraService2.extraPort | nindent 4 }}
+  {{- end }}
   selector: {{- include "common.labels.matchLabels" . | nindent 4 }}
 {{- end }}
\ No newline at end of file
diff --git a/Engine/Source/Programs/UnrealCloudDDC/Helm/UnrealCloudDDC/templates/extra-service.yaml b/Engine/Source/Programs/UnrealCloudDDC/Helm/UnrealCloudDDC/templates/extra-service.yaml
--- a/Engine/Source/Programs/UnrealCloudDDC/Helm/UnrealCloudDDC/templates/extra-service.yaml
+++ b/Engine/Source/Programs/UnrealCloudDDC/Helm/UnrealCloudDDC/templates/extra-service.yaml
@@ -17,5 +17,8 @@ spec:
     - name: {{ .Values.extraService.portName }}
       port: {{ .Values.extraService.port }}
       targetPort: {{ default "http" .Values.extraService.targetPort }}
+  {{- if .Values.extraService.extraPort }}
+  {{- toYaml .Values.extraService.extraPort | nindent 4 }}
+  {{- end }}
   selector: {{- include "common.labels.matchLabels" . | nindent 4 }}
 {{- end }}
\ No newline at end of file
diff --git a/Engine/Source/Programs/UnrealCloudDDC/Helm/UnrealCloudDDC/templates/ingress.yaml b/Engine/Source/Programs/UnrealCloudDDC/Helm/UnrealCloudDDC/templates/ingress.yaml
--- a/Engine/Source/Programs/UnrealCloudDDC/Helm/UnrealCloudDDC/templates/ingress.yaml
+++ b/Engine/Source/Programs/UnrealCloudDDC/Helm/UnrealCloudDDC/templates/ingress.yaml
@@ -1,4 +1,5 @@
 {{- if .Values.ingress.enabled }}
+{{- $servicePort := default "http" .Values.ingress.servicePort -}}
 apiVersion: {{ include "common.capabilities.ingress.apiVersion" . }}
 kind: Ingress
 metadata:
@@ -24,7 +25,7 @@ spec:
             {{- if eq "true" (include "common.ingress.supportsPathType" .) }}
             pathType: {{ .Values.ingress.pathType }}
             {{- end }}
-            backend: {{- include "common.ingress.backend" (dict "serviceName" (include "common.names.fullname" .) "servicePort" "http" "context" $)  | nindent 14 }}
+            backend: {{- include "common.ingress.backend" (dict "serviceName" (include "common.names.fullname" .) "servicePort" $servicePort "context" $)  | nindent 14 }}
     {{- end }}
     {{- range .Values.ingress.extraHosts }}
     - host: {{ .name }}
@@ -34,7 +35,7 @@ spec:
             {{- if eq "true" (include "common.ingress.supportsPathType" $) }}
             pathType: {{ default "ImplementationSpecific" .pathType }}
             {{- end }}
-            backend: {{- include "common.ingress.backend" (dict "serviceName" (include "common.names.fullname" $) "servicePort" "http" "context" $) | nindent 14 }}
+            backend: {{- include "common.ingress.backend" (dict "serviceName" (include "common.names.fullname" $) "servicePort" $servicePort "context" $) | nindent 14 }}
     {{- end }}
   {{- if or .Values.ingress.tls .Values.ingress.extraTls .Values.ingress.hosts }}
   tls:
diff --git a/Engine/Source/Programs/UnrealCloudDDC/Helm/UnrealCloudDDC/values.yaml b/Engine/Source/Programs/UnrealCloudDDC/Helm/UnrealCloudDDC/values.yaml
--- a/Engine/Source/Programs/UnrealCloudDDC/Helm/UnrealCloudDDC/values.yaml
+++ b/Engine/Source/Programs/UnrealCloudDDC/Helm/UnrealCloudDDC/values.yaml
@@ -307,6 +307,15 @@ useDotnetMonitorSidecar : false
 ## this is typically faster then falling back to global blob store
 fetchBlobsFromPeerPods: false
 
+## Nginx configuration
+nginx:
+
+  ## Enable this option to run a nginx proxy in the all pods that is able to serve content directly from the filesystem
+  enabled: false
+
+  ## Set this to add a header whenever X-Accel is used from nginx to send blobs
+  addXAccelDebugHeader: false
+
 ## Worker configuration
 
 worker:
