Unreal Cloud DDC - Added configmap to nginx proxy

diff --git a/Engine/Source/Programs/UnrealCloudDDC/Helm/UnrealCloudDDC/templates/configmap-nginx.yaml b/Engine/Source/Programs/UnrealCloudDDC/Helm/UnrealCloudDDC/templates/configmap-nginx.yaml
new file mode 100644
--- /dev/null
+++ b/Engine/Source/Programs/UnrealCloudDDC/Helm/UnrealCloudDDC/templates/configmap-nginx.yaml
@@ -0,0 +1,110 @@
+{{- if .Values.nginx.enabled }}
+apiVersion: v1
+kind: ConfigMap
+metadata:
+  name: {{ include "common.names.fullname" . }}-nginx
+  labels: {{- include "common.labels.standard" . | nindent 4 }}
+data:
+  nginx.conf: |
+    worker_processes auto;
+
+    events {
+      worker_connections  2048;
+    }
+
+    http {
+      access_log off;
+      client_max_body_size 0;
+      server_tokens off;
+      sendfile on;
+      tcp_nopush on;
+
+      # http 1.1 public port
+      server {
+        listen 81 backlog=4096;
+        server_name public_http1;
+
+        proxy_buffering off;
+        proxy_http_version 1.1;
+        
+        location /api {
+            proxy_set_header X-Jupiter-XAccel-Supported true;
+            proxy_pass http://localhost:80;
+        }
+        location /health {
+            proxy_pass http://localhost:80;
+        }
+        location /nginx-blobs {
+            internal;
+            alias {{ printf "%s/Blobs" .Values.persistence.mountPath }};
+            {{ if .Values.nginx.addXAccelDebugHeader }} add_header X-Jupiter-XAccel-Used true; {{ end }}
+        }
+      }
+
+      # http 1.1 corp port
+      server {
+        listen 8009 backlog=4096;
+        server_name corp_http1;
+        
+        proxy_buffering off;
+        proxy_http_version 1.1;
+        
+        location /api {
+            proxy_set_header X-Jupiter-XAccel-Supported true;
+            proxy_pass http://localhost:8008;
+        }
+        location /health {
+            proxy_pass http://localhost:8008;
+        }
+        location /nginx-blobs {
+            internal;
+            alias {{ printf "%s/Blobs" .Values.persistence.mountPath }};
+            {{ if .Values.nginx.addXAccelDebugHeader }} add_header X-Jupiter-XAccel-Used true; {{ end }}
+        }
+      }
+
+      # public http2 port
+      server {
+        listen 8091 backlog=4096 http2;
+        server_name public_http2;
+
+        proxy_buffering off;
+        proxy_http_version 1.1;
+        
+        location /api {
+            proxy_set_header X-Jupiter-XAccel-Supported true;
+            proxy_pass http://localhost:80;
+        }
+        location /health {
+            proxy_pass http://localhost:80;
+        }
+        location /nginx-blobs {
+            internal;
+            alias {{ printf "%s/Blobs" .Values.persistence.mountPath }};
+            {{ if .Values.nginx.addXAccelDebugHeader }} add_header X-Jupiter-XAccel-Used true; {{ end }}
+        }
+      }
+
+      # corp port http2
+      server {
+        listen 8092 backlog=4096 http2;
+        server_name corp_http2;
+
+        proxy_buffering off;
+        proxy_http_version 1.1;
+        
+        location /api {
+            proxy_set_header X-Jupiter-XAccel-Supported true;
+            proxy_pass http://localhost:8008;
+        }
+        location /health {
+            proxy_pass http://localhost:8008;
+        }
+        location /nginx-blobs {
+            internal;
+            alias {{ printf "%s/Blobs" .Values.persistence.mountPath }};
+            {{ if .Values.nginx.addXAccelDebugHeader }}add_header X-Jupiter-XAccel-Used true; {{ end }}
+        }
+      }
+    }
+{{ end }}
