Unreal Cloud DDC - Added ability to accept connection via unix domain sockets

diff --git a/Engine/Source/Programs/UnrealCloudDDC/Jupiter.Common/BaseProgram.cs b/Engine/Source/Programs/UnrealCloudDDC/Jupiter.Common/BaseProgram.cs
--- a/Engine/Source/Programs/UnrealCloudDDC/Jupiter.Common/BaseProgram.cs
+++ b/Engine/Source/Programs/UnrealCloudDDC/Jupiter.Common/BaseProgram.cs
@@ -3,11 +3,14 @@
 using System;
 using System.Diagnostics;
 using System.IO;
+using EpicGames.Core;
 using Microsoft.AspNetCore.Hosting;
+using Microsoft.AspNetCore.Server.Kestrel.Core;
 using Microsoft.Extensions.Configuration;
 using Microsoft.Extensions.Hosting;
 using Serilog;
 using Serilog.Core;
+using Log = Serilog.Log;
 
 namespace Jupiter
 {
@@ -54,7 +57,22 @@ namespace Jupiter
             try
             {
                 Log.Information("Creating ASPNET Host");
-                CreateHostBuilder(args).Build().Run();
+                IHost host = CreateHostBuilder(args).Build();
+                host.Start();
+
+                JupiterSettings settings = new JupiterSettings();
+                Configuration.GetSection("Jupiter").Bind(settings);
+                string socketsRoot = settings.DomainSocketsRoot;
+
+                if (settings.ChmodDomainSockets)
+                {
+                    FileUtils.SetFileMode_Linux(Path.Combine(socketsRoot, "public-http.sock"), 755);
+                    FileUtils.SetFileMode_Linux(Path.Combine(socketsRoot, "public-http2.sock"), 755);
+                    FileUtils.SetFileMode_Linux(Path.Combine(socketsRoot, "corp-http.sock"), 755);
+                    FileUtils.SetFileMode_Linux(Path.Combine(socketsRoot, "corp-http2.sock"), 755);
+                }
+
+                host.WaitForShutdown();
                 return 0;
             }
             catch (Exception ex)
@@ -89,6 +107,33 @@ namespace Jupiter
                     webBuilder.ConfigureKestrel(options =>
                     {
                         options.AddServerHeader = false;
+
+                        JupiterSettings settings = new JupiterSettings();
+                        Configuration.GetSection("Jupiter").Bind(settings);
+                        string socketsRoot = settings.DomainSocketsRoot;
+
+                        if (settings.UseDomainSockets)
+                        {
+                            Log.Logger.Information("Using unix domain sockets at {SocketsRoot}", socketsRoot);
+                            options.ListenUnixSocket(Path.Combine(socketsRoot, "public-http.sock"), listenOptions =>
+                            {
+                                listenOptions.Protocols = HttpProtocols.Http1AndHttp2;
+                            });
+                            options.ListenUnixSocket(Path.Combine(socketsRoot, "public-http2.sock"), listenOptions =>
+                            {
+                                listenOptions.Protocols = HttpProtocols.Http2;
+                            });
+
+                            options.ListenUnixSocket(Path.Combine(socketsRoot, "corp-http.sock"), listenOptions =>
+                            {
+                                listenOptions.Protocols = HttpProtocols.Http1AndHttp2;
+                            });
+
+                            options.ListenUnixSocket(Path.Combine(socketsRoot, "corp-http2.sock"), listenOptions =>
+                            {
+                                listenOptions.Protocols = HttpProtocols.Http2;
+                            });
+                        }
                     });
                 });
         }
diff --git a/Engine/Source/Programs/UnrealCloudDDC/Jupiter.Common/BaseStartup.cs b/Engine/Source/Programs/UnrealCloudDDC/Jupiter.Common/BaseStartup.cs
--- a/Engine/Source/Programs/UnrealCloudDDC/Jupiter.Common/BaseStartup.cs
+++ b/Engine/Source/Programs/UnrealCloudDDC/Jupiter.Common/BaseStartup.cs
@@ -711,6 +711,21 @@ namespace Jupiter
         [Required]
         [Key]
         public string CurrentSite { get; set; } = "";
+
+        /// <summary>
+        /// Used to move where domain sockets are allocated
+        /// </summary>
+        public string DomainSocketsRoot { get; set; } = "/tmp/sockets";
+
+        /// <summary>
+        /// Enable to create unix domain sockets for inter process communication
+        /// </summary>
+        public bool UseDomainSockets { get; set; } = false;
+
+        /// <summary>
+        /// Enable to change access (chmod) the sockets to allow for anyone to access them
+        /// </summary>
+        public bool ChmodDomainSockets { get; set; } = false;
     }
 
     public class NamespaceSettings
