Unreal Cloud DDC - Fixed issue with domain sockets not being correctly identified as public or corp, added header which can be sent if on a local connection to indicate port used (which can be sent from nginx).
Fixed up assumption that all local connections are test (with domain sockets this is no longer the case).
Reconciled our domain sockets into 2, one for http1.1 and one for http2.

diff --git a/Engine/Source/Programs/UnrealCloudDDC/Jupiter.Common/BaseProgram.cs b/Engine/Source/Programs/UnrealCloudDDC/Jupiter.Common/BaseProgram.cs
--- a/Engine/Source/Programs/UnrealCloudDDC/Jupiter.Common/BaseProgram.cs
+++ b/Engine/Source/Programs/UnrealCloudDDC/Jupiter.Common/BaseProgram.cs
@@ -66,10 +66,8 @@ namespace Jupiter
 
                 if (settings.ChmodDomainSockets)
                 {
-                    FileUtils.SetFileMode_Linux(Path.Combine(socketsRoot, "public-http.sock"), 755);
-                    FileUtils.SetFileMode_Linux(Path.Combine(socketsRoot, "public-http2.sock"), 755);
-                    FileUtils.SetFileMode_Linux(Path.Combine(socketsRoot, "corp-http.sock"), 755);
-                    FileUtils.SetFileMode_Linux(Path.Combine(socketsRoot, "corp-http2.sock"), 755);
+                    FileUtils.SetFileMode_Linux(Path.Combine(socketsRoot, "jupiter-http.sock"), 766);
+                    FileUtils.SetFileMode_Linux(Path.Combine(socketsRoot, "jupiter-http2.sock"), 766);
                 }
 
                 host.WaitForShutdown();
@@ -115,21 +113,11 @@ namespace Jupiter
                         if (settings.UseDomainSockets)
                         {
                             Log.Logger.Information("Using unix domain sockets at {SocketsRoot}", socketsRoot);
-                            options.ListenUnixSocket(Path.Combine(socketsRoot, "public-http.sock"), listenOptions =>
+                            options.ListenUnixSocket(Path.Combine(socketsRoot, "jupiter-http.sock"), listenOptions =>
                             {
                                 listenOptions.Protocols = HttpProtocols.Http1AndHttp2;
                             });
-                            options.ListenUnixSocket(Path.Combine(socketsRoot, "public-http2.sock"), listenOptions =>
-                            {
-                                listenOptions.Protocols = HttpProtocols.Http2;
-                            });
-
-                            options.ListenUnixSocket(Path.Combine(socketsRoot, "corp-http.sock"), listenOptions =>
-                            {
-                                listenOptions.Protocols = HttpProtocols.Http1AndHttp2;
-                            });
-
-                            options.ListenUnixSocket(Path.Combine(socketsRoot, "corp-http2.sock"), listenOptions =>
+                            options.ListenUnixSocket(Path.Combine(socketsRoot, "jupiter-http2.sock"), listenOptions =>
                             {
                                 listenOptions.Protocols = HttpProtocols.Http2;
                             });
diff --git a/Engine/Source/Programs/UnrealCloudDDC/Jupiter.Common/BaseStartup.cs b/Engine/Source/Programs/UnrealCloudDDC/Jupiter.Common/BaseStartup.cs
--- a/Engine/Source/Programs/UnrealCloudDDC/Jupiter.Common/BaseStartup.cs
+++ b/Engine/Source/Programs/UnrealCloudDDC/Jupiter.Common/BaseStartup.cs
@@ -726,6 +726,11 @@ namespace Jupiter
         /// Enable to change access (chmod) the sockets to allow for anyone to access them
         /// </summary>
         public bool ChmodDomainSockets { get; set; } = false;
+
+        /// <summary>
+        /// Assumes that any local connection should have full access, this is used only for tests
+        /// </summary>
+        public bool AssumeLocalConnectionsHasFullAccess { get; set; }
     }
 
     public class NamespaceSettings
diff --git a/Engine/Source/Programs/UnrealCloudDDC/Jupiter/Controllers/RequestHelper.cs b/Engine/Source/Programs/UnrealCloudDDC/Jupiter/Controllers/RequestHelper.cs
--- a/Engine/Source/Programs/UnrealCloudDDC/Jupiter/Controllers/RequestHelper.cs
+++ b/Engine/Source/Programs/UnrealCloudDDC/Jupiter/Controllers/RequestHelper.cs
@@ -9,6 +9,7 @@ using Microsoft.AspNetCore.Authorization;
 using Microsoft.AspNetCore.Http;
 using Microsoft.AspNetCore.Mvc;
 using Microsoft.Extensions.Options;
+using Microsoft.Extensions.Primitives;
 using OpenTelemetry.Trace;
 
 namespace Jupiter.Controllers;
@@ -52,11 +53,36 @@ public class RequestHelper
 
         // namespace is a restricted namespace
         HttpContext context = request.HttpContext;
-        bool isUnitTest = context.Connection.LocalPort == 0 && context.Connection.LocalIpAddress == null;
-        /* unit tests do not run on ports, we consider them always on the internal port */
+        string? portHeaderValue = null;
+        if (context.Request.Headers.TryGetValue("X-Jupiter-Port", out StringValues values))
+        {
+            portHeaderValue = values.ToString();
+        }
+
+        // unit tests do not run on ports, we consider them always on the internal port
+        bool isLocalConnection = context.Connection.LocalPort == 0 && context.Connection.LocalIpAddress == null;
+        // public port is either running on the public port, or if using domain sockets we check the header that is passed along instead
         bool isPublicPort = _settings!.CurrentValue.PublicApiPorts.Contains(context.Connection.LocalPort);
-            
-        if (isPublicPort && !isUnitTest)
+
+        if (isLocalConnection && _settings.CurrentValue.AssumeLocalConnectionsHasFullAccess)
+        {
+            // local connection so granting it full access
+            isPublicPort = false;
+        }
+
+        if (isLocalConnection && portHeaderValue != null)
+        {
+            if (string.Equals(portHeaderValue, "Public", StringComparison.OrdinalIgnoreCase))
+            {
+                isPublicPort = true;
+            }
+            else if (string.Equals(portHeaderValue, "Corp", StringComparison.OrdinalIgnoreCase))
+            {
+                isPublicPort = false;
+            }
+        }
+
+        if (isPublicPort)
         {
             // trying to access restricted namespace on a public port, this is not allowed
             return new ForbidResult();
diff --git a/Engine/Source/Programs/UnrealCloudDDC/Jupiter/appsettings.Testing.json b/Engine/Source/Programs/UnrealCloudDDC/Jupiter/appsettings.Testing.json
--- a/Engine/Source/Programs/UnrealCloudDDC/Jupiter/appsettings.Testing.json
+++ b/Engine/Source/Programs/UnrealCloudDDC/Jupiter/appsettings.Testing.json
@@ -220,7 +220,8 @@
   "AllowedHosts": "*",
 
   "Jupiter": {
-    "CurrentSite": "test"
+    "CurrentSite": "test",
+    "AssumeLocalConnectionsHasFullAccess": "true"
   },
 
   "Kestrel": {
