Unreal Cloud DDC - Added options in helm to configure nginx to use unix domain sockets to jupiter

diff --git a/Engine/Source/Programs/UnrealCloudDDC/Helm/UnrealCloudDDC/templates/configmap-nginx.yaml b/Engine/Source/Programs/UnrealCloudDDC/Helm/UnrealCloudDDC/templates/configmap-nginx.yaml
--- a/Engine/Source/Programs/UnrealCloudDDC/Helm/UnrealCloudDDC/templates/configmap-nginx.yaml
+++ b/Engine/Source/Programs/UnrealCloudDDC/Helm/UnrealCloudDDC/templates/configmap-nginx.yaml
@@ -28,16 +28,18 @@ data:
         proxy_http_version 1.1;
         
         location /api {
-            proxy_set_header X-Jupiter-XAccel-Supported true;
-            proxy_pass http://localhost:80;
+          proxy_set_header X-Jupiter-XAccel-Supported true;
+          proxy_set_header X-Jupiter-Port Public;
+          proxy_pass {{ if .Values.nginx.useDomainSockets -}} http://unix:/nginx/jupiter-http.sock {{- else -}} http://localhost:80 {{- end }};
         }
         location /health {
-            proxy_pass http://localhost:80;
+          proxy_set_header X-Jupiter-Port Public;
+          proxy_pass {{ if .Values.nginx.useDomainSockets -}} http://unix:/nginx/jupiter-http.sock {{- else -}} http://localhost:80 {{- end }};
         }
         location /nginx-blobs {
-            internal;
-            alias {{ printf "%s/Blobs" .Values.persistence.mountPath }};
-            {{ if .Values.nginx.addXAccelDebugHeader }} add_header X-Jupiter-XAccel-Used true; {{ end }}
+          internal;
+          alias {{ printf "%s/Blobs" .Values.persistence.mountPath }};
+          {{ if .Values.nginx.addXAccelDebugHeader }} add_header X-Jupiter-XAccel-Used true; {{ end }}
         }
       }
 
@@ -50,16 +52,18 @@ data:
         proxy_http_version 1.1;
         
         location /api {
-            proxy_set_header X-Jupiter-XAccel-Supported true;
-            proxy_pass http://localhost:8008;
+          proxy_set_header X-Jupiter-XAccel-Supported true;
+          proxy_set_header X-Jupiter-Port Corp;
+          proxy_pass {{ if .Values.nginx.useDomainSockets -}} http://unix:/nginx/jupiter-http.sock {{- else -}} http://localhost:8008 {{- end }};
         }
         location /health {
-            proxy_pass http://localhost:8008;
+          proxy_set_header X-Jupiter-Port Corp;
+          proxy_pass {{ if .Values.nginx.useDomainSockets -}} http://unix:/nginx/jupiter-http.sock {{- else -}} http://localhost:8008 {{- end }};
         }
         location /nginx-blobs {
-            internal;
-            alias {{ printf "%s/Blobs" .Values.persistence.mountPath }};
-            {{ if .Values.nginx.addXAccelDebugHeader }} add_header X-Jupiter-XAccel-Used true; {{ end }}
+          internal;
+          alias {{ printf "%s/Blobs" .Values.persistence.mountPath }};
+          {{ if .Values.nginx.addXAccelDebugHeader }} add_header X-Jupiter-XAccel-Used true; {{ end }}
         }
       }
 
@@ -72,16 +76,18 @@ data:
         proxy_http_version 1.1;
         
         location /api {
-            proxy_set_header X-Jupiter-XAccel-Supported true;
-            proxy_pass http://localhost:80;
+          proxy_set_header X-Jupiter-XAccel-Supported true;
+          proxy_set_header X-Jupiter-Port Public;
+          proxy_pass {{ if .Values.nginx.useDomainSockets -}} http://unix:/nginx/jupiter-http2.sock {{- else -}} http://localhost:80 {{- end }};
         }
         location /health {
-            proxy_pass http://localhost:80;
+          proxy_set_header X-Jupiter-Port Public;
+          proxy_pass {{ if .Values.nginx.useDomainSockets -}} http://unix:/nginx/jupiter-http2.sock {{- else -}} http://localhost:80 {{- end }};
         }
         location /nginx-blobs {
-            internal;
-            alias {{ printf "%s/Blobs" .Values.persistence.mountPath }};
-            {{ if .Values.nginx.addXAccelDebugHeader }} add_header X-Jupiter-XAccel-Used true; {{ end }}
+          internal;
+          alias {{ printf "%s/Blobs" .Values.persistence.mountPath }};
+          {{ if .Values.nginx.addXAccelDebugHeader }} add_header X-Jupiter-XAccel-Used true; {{ end }}
         }
       }
 
@@ -94,11 +100,13 @@ data:
         proxy_http_version 1.1;
         
         location /api {
-            proxy_set_header X-Jupiter-XAccel-Supported true;
-            proxy_pass http://localhost:8008;
+          proxy_set_header X-Jupiter-XAccel-Supported true;
+          proxy_set_header X-Jupiter-Port Corp;
+          proxy_pass {{ if .Values.nginx.useDomainSockets -}} http://unix:/nginx/jupiter-http.sock {{- else -}} http://localhost:8008 {{- end }};
         }
         location /health {
-            proxy_pass http://localhost:8008;
+          proxy_set_header X-Jupiter-Port Corp;
+          proxy_pass {{ if .Values.nginx.useDomainSockets -}} http://unix:/nginx/jupiter-http.sock {{- else -}} http://localhost:8008 {{- end }};
         }
         location /nginx-blobs {
             internal;
diff --git a/Engine/Source/Programs/UnrealCloudDDC/Helm/UnrealCloudDDC/templates/configmap.yaml b/Engine/Source/Programs/UnrealCloudDDC/Helm/UnrealCloudDDC/templates/configmap.yaml
--- a/Engine/Source/Programs/UnrealCloudDDC/Helm/UnrealCloudDDC/templates/configmap.yaml
+++ b/Engine/Source/Programs/UnrealCloudDDC/Helm/UnrealCloudDDC/templates/configmap.yaml
@@ -64,6 +64,13 @@ Jupiter:
 {{ $config.Jupiter | toYaml | nindent 2 }}
 {{- end }}
 
+{{ if .Values.nginx.useDomainSockets }}
+Jupiter:
+  UseDomainSockets: true
+  DomainSocketsRoot: "/nginx"
+  ChmodDomainSockets: true
+{{ end }}
+
 {{ if .Values.nginx.enabled }}
 Nginx:
   UseNginxRedirect: true
diff --git a/Engine/Source/Programs/UnrealCloudDDC/Helm/UnrealCloudDDC/values.yaml b/Engine/Source/Programs/UnrealCloudDDC/Helm/UnrealCloudDDC/values.yaml
--- a/Engine/Source/Programs/UnrealCloudDDC/Helm/UnrealCloudDDC/values.yaml
+++ b/Engine/Source/Programs/UnrealCloudDDC/Helm/UnrealCloudDDC/values.yaml
@@ -316,6 +316,9 @@ nginx:
   ## Set this to add a header whenever X-Accel is used from nginx to send blobs
   addXAccelDebugHeader: false
 
+  ## Set this to use Unix domain sockets rather then ports for inter process communication
+  useDomainSockets: false
+
 ## Worker configuration
 
 worker:
