Unreal Cloud DDC - Remove html error page response from nginx for 404

diff --git a/Engine/Source/Programs/UnrealCloudDDC/Helm/UnrealCloudDDC/templates/configmap-nginx.yaml b/Engine/Source/Programs/UnrealCloudDDC/Helm/UnrealCloudDDC/templates/configmap-nginx.yaml
--- a/Engine/Source/Programs/UnrealCloudDDC/Helm/UnrealCloudDDC/templates/configmap-nginx.yaml
+++ b/Engine/Source/Programs/UnrealCloudDDC/Helm/UnrealCloudDDC/templates/configmap-nginx.yaml
@@ -41,6 +41,12 @@ data:
           alias {{ printf "%s/Blobs" .Values.persistence.mountPath }};
           {{ if .Values.nginx.addXAccelDebugHeader }} add_header X-Jupiter-XAccel-Used true; {{ end }}
         }
+
+        error_page 404 @error_page;
+        location @error_page {
+          default_type "";
+          return 200; # error_code is ignored
+        }
       }
 
       # http 1.1 corp port
@@ -65,6 +71,12 @@ data:
           alias {{ printf "%s/Blobs" .Values.persistence.mountPath }};
           {{ if .Values.nginx.addXAccelDebugHeader }} add_header X-Jupiter-XAccel-Used true; {{ end }}
         }
+
+        error_page 404 @error_page;
+        location @error_page {
+          default_type "";
+          return 200; # error_code is ignored
+        }
       }
 
       # public http2 port
@@ -89,6 +101,12 @@ data:
           alias {{ printf "%s/Blobs" .Values.persistence.mountPath }};
           {{ if .Values.nginx.addXAccelDebugHeader }} add_header X-Jupiter-XAccel-Used true; {{ end }}
         }
+
+        error_page 404 @error_page;
+        location @error_page {
+          default_type "";
+          return 200; # error_code is ignored
+        }
       }
 
       # corp port http2
@@ -113,6 +131,12 @@ data:
             alias {{ printf "%s/Blobs" .Values.persistence.mountPath }};
             {{ if .Values.nginx.addXAccelDebugHeader }}add_header X-Jupiter-XAccel-Used true; {{ end }}
         }
+
+        error_page 404 @error_page;
+        location @error_page {
+          default_type "";
+          return 200; # error_code is ignored
+        }
       }
     }
 {{ end }}
