Unreal Cloud DDC - Always direct traffic to the http1.1 port as nginx itself will never use http2 to talk to our backend.

diff --git a/Engine/Source/Programs/UnrealCloudDDC/Helm/UnrealCloudDDC/templates/configmap-nginx.yaml b/Engine/Source/Programs/UnrealCloudDDC/Helm/UnrealCloudDDC/templates/configmap-nginx.yaml
--- a/Engine/Source/Programs/UnrealCloudDDC/Helm/UnrealCloudDDC/templates/configmap-nginx.yaml
+++ b/Engine/Source/Programs/UnrealCloudDDC/Helm/UnrealCloudDDC/templates/configmap-nginx.yaml
@@ -90,11 +90,11 @@ data:
         location /api {
           proxy_set_header X-Jupiter-XAccel-Supported true;
           proxy_set_header X-Jupiter-Port Public;
-          proxy_pass {{ if .Values.nginx.useDomainSockets -}} http://unix:/nginx/jupiter-http2.sock {{- else -}} http://localhost:80 {{- end }};
+          proxy_pass {{ if .Values.nginx.useDomainSockets -}} http://unix:/nginx/jupiter-http.sock {{- else -}} http://localhost:80 {{- end }};
         }
         location /health {
           proxy_set_header X-Jupiter-Port Public;
-          proxy_pass {{ if .Values.nginx.useDomainSockets -}} http://unix:/nginx/jupiter-http2.sock {{- else -}} http://localhost:80 {{- end }};
+          proxy_pass {{ if .Values.nginx.useDomainSockets -}} http://unix:/nginx/jupiter-http.sock {{- else -}} http://localhost:80 {{- end }};
         }
         location /nginx-blobs {
           internal;
