Unreal Cloud DDC - Fixed issue with cb package flags not being set correctly when fetching them

diff --git a/Engine/Source/Programs/UnrealCloudDDC/Jupiter/Controllers/ReferencesController.cs b/Engine/Source/Programs/UnrealCloudDDC/Jupiter/Controllers/ReferencesController.cs
--- a/Engine/Source/Programs/UnrealCloudDDC/Jupiter/Controllers/ReferencesController.cs
+++ b/Engine/Source/Programs/UnrealCloudDDC/Jupiter/Controllers/ReferencesController.cs
@@ -257,7 +257,7 @@ namespace Jupiter.Controllers
                                 }
                                 else if (attachment is ObjectAttachment objectAttachment)
                                 {
-                                    flags &= CbPackageAttachmentFlags.IsObject;
+                                    flags |= CbPackageAttachmentFlags.IsObject;
                                     BlobIdentifier referencedBlob = objectAttachment.Identifier;
                                     attachmentContents = await _blobStore.GetObject(ns, referencedBlob);
                                 }
@@ -268,7 +268,7 @@ namespace Jupiter.Controllers
                                     (attachmentContents, string mime) = await _blobStore.GetCompressedObject(ns, contentId, HttpContext.RequestServices);
                                     if (mime == CustomMediaTypeNames.UnrealCompressedBuffer)
                                     {
-                                        flags &= CbPackageAttachmentFlags.IsCompressed;
+                                        flags |= CbPackageAttachmentFlags.IsCompressed;
                                     }
                                     else
                                     {
@@ -718,7 +718,8 @@ namespace Jupiter.Controllers
 
             _diagnosticContext.Set("Content-Length", Request.ContentLength ?? -1);
 
-            CbPackageReader packageReader = await CbPackageReader.Create(Request.Body);
+            byte[] b = await RequestUtil.ReadRawBody(Request);
+            CbPackageReader packageReader = await CbPackageReader.Create(new MemoryStream(b));
 
             try
             {
