Unreal Cloud DDC - Fixed name of nginx http2 corp port

diff --git a/Engine/Source/Programs/UnrealCloudDDC/Helm/UnrealCloudDDC/templates/deployment.yaml b/Engine/Source/Programs/UnrealCloudDDC/Helm/UnrealCloudDDC/templates/deployment.yaml
--- a/Engine/Source/Programs/UnrealCloudDDC/Helm/UnrealCloudDDC/templates/deployment.yaml
+++ b/Engine/Source/Programs/UnrealCloudDDC/Helm/UnrealCloudDDC/templates/deployment.yaml
@@ -134,7 +134,7 @@ spec:
           containerPort: 8009
         - name: nginx-h2
           containerPort: 8091
-        - name: nginx-corp-h
+        - name: nginx-corp-h2
           containerPort: 8092
         volumeMounts:
           - name: nginx
