Unreal Cloud DDC - Added health checks to the nginx container

diff --git a/Engine/Source/Programs/UnrealCloudDDC/Helm/UnrealCloudDDC/templates/deployment.yaml b/Engine/Source/Programs/UnrealCloudDDC/Helm/UnrealCloudDDC/templates/deployment.yaml
--- a/Engine/Source/Programs/UnrealCloudDDC/Helm/UnrealCloudDDC/templates/deployment.yaml
+++ b/Engine/Source/Programs/UnrealCloudDDC/Helm/UnrealCloudDDC/templates/deployment.yaml
@@ -136,6 +136,28 @@ spec:
           containerPort: 8091
         - name: nginx-corp-h2
           containerPort: 8092
+        {{- if .Values.nginx.livenessProbe.enabled }}
+        livenessProbe:
+          httpGet:
+            path: {{ default "health/liveness" .Values.nginx.livenessProbe.path }}
+            port: nginx-http
+          initialDelaySeconds: {{ .Values.nginx.livenessProbe.initialDelaySeconds }}
+          periodSeconds: {{ .Values.nginx.livenessProbe.periodSeconds }}
+          timeoutSeconds: {{ .Values.nginx.livenessProbe.timeoutSeconds }}
+          successThreshold: {{ .Values.nginx.livenessProbe.successThreshold }}
+          failureThreshold: {{ .Values.nginx.livenessProbe.failureThreshold }}
+        {{- end }}
+        {{- if .Values.nginx.readinessProbe.enabled }}
+        readinessProbe:
+          httpGet:
+            path:  {{ default "health/readiness" .Values.nginx.readinessProbe.path }}
+            port: nginx-http
+          initialDelaySeconds: {{ .Values.nginx.readinessProbe.initialDelaySeconds }}
+          periodSeconds: {{ .Values.nginx.readinessProbe.periodSeconds }}
+          timeoutSeconds: {{ .Values.nginx.readinessProbe.timeoutSeconds }}
+          successThreshold: {{ .Values.nginx.readinessProbe.successThreshold }}
+          failureThreshold: {{ .Values.nginx.readinessProbe.failureThreshold }}
+        {{- end }}
         volumeMounts:
           - name: nginx
             mountPath: /nginx
diff --git a/Engine/Source/Programs/UnrealCloudDDC/Helm/UnrealCloudDDC/values.yaml b/Engine/Source/Programs/UnrealCloudDDC/Helm/UnrealCloudDDC/values.yaml
--- a/Engine/Source/Programs/UnrealCloudDDC/Helm/UnrealCloudDDC/values.yaml
+++ b/Engine/Source/Programs/UnrealCloudDDC/Helm/UnrealCloudDDC/values.yaml
@@ -319,8 +319,24 @@ nginx:
   ## Set this to use Unix domain sockets rather then ports for inter process communication
   useDomainSockets: false
 
-## Worker configuration
+  ## Liveness and readiness probes, defaults to corresponding endpoints
+  ##
+  livenessProbe:
+    enabled: true
+    initialDelaySeconds: 5
+    periodSeconds: 5
+    timeoutSeconds: 5
+    successThreshold: 1
+    failureThreshold: 5
+  readinessProbe:
+    enabled: true
+    initialDelaySeconds: 5
+    periodSeconds: 5
+    timeoutSeconds: 1
+    successThreshold: 1
+    failureThreshold: 5
 
+## Worker configuration
 worker:
 
   ## Set to enable seperate deployments for background workers
