Unreal Cloud DDC - Use the standard http port from nginx when not using domain sockets (as nginx always uses http 1.1 to connect to the backend)

diff --git a/Engine/Source/Programs/UnrealCloudDDC/Helm/UnrealCloudDDC/templates/configmap-nginx.yaml b/Engine/Source/Programs/UnrealCloudDDC/Helm/UnrealCloudDDC/templates/configmap-nginx.yaml
--- a/Engine/Source/Programs/UnrealCloudDDC/Helm/UnrealCloudDDC/templates/configmap-nginx.yaml
+++ b/Engine/Source/Programs/UnrealCloudDDC/Helm/UnrealCloudDDC/templates/configmap-nginx.yaml
@@ -60,11 +60,11 @@ data:
         location /api {
           proxy_set_header X-Jupiter-XAccel-Supported true;
           proxy_set_header X-Jupiter-Port Corp;
-          proxy_pass {{ if .Values.nginx.useDomainSockets -}} http://unix:/nginx/jupiter-http.sock {{- else -}} http://localhost:8008 {{- end }};
+          proxy_pass {{ if .Values.nginx.useDomainSockets -}} http://unix:/nginx/jupiter-http.sock {{- else -}} http://localhost:80 {{- end }};
         }
         location /health {
           proxy_set_header X-Jupiter-Port Corp;
-          proxy_pass {{ if .Values.nginx.useDomainSockets -}} http://unix:/nginx/jupiter-http.sock {{- else -}} http://localhost:8008 {{- end }};
+          proxy_pass {{ if .Values.nginx.useDomainSockets -}} http://unix:/nginx/jupiter-http.sock {{- else -}} http://localhost:80 {{- end }};
         }
         location /nginx-blobs {
           internal;
@@ -120,11 +120,11 @@ data:
         location /api {
           proxy_set_header X-Jupiter-XAccel-Supported true;
           proxy_set_header X-Jupiter-Port Corp;
-          proxy_pass {{ if .Values.nginx.useDomainSockets -}} http://unix:/nginx/jupiter-http.sock {{- else -}} http://localhost:8008 {{- end }};
+          proxy_pass {{ if .Values.nginx.useDomainSockets -}} http://unix:/nginx/jupiter-http.sock {{- else -}} http://localhost:80 {{- end }};
         }
         location /health {
           proxy_set_header X-Jupiter-Port Corp;
-          proxy_pass {{ if .Values.nginx.useDomainSockets -}} http://unix:/nginx/jupiter-http.sock {{- else -}} http://localhost:8008 {{- end }};
+          proxy_pass {{ if .Values.nginx.useDomainSockets -}} http://unix:/nginx/jupiter-http.sock {{- else -}} http://localhost:80 {{- end }};
         }
         location /nginx-blobs {
             internal;
