Unreal Cloud DDC - Added nginx monitoring endpoints
Also added ability to tweak worker connections in nginx

diff --git a/Engine/Source/Programs/UnrealCloudDDC/Helm/UnrealCloudDDC/templates/configmap-nginx.yaml b/Engine/Source/Programs/UnrealCloudDDC/Helm/UnrealCloudDDC/templates/configmap-nginx.yaml
--- a/Engine/Source/Programs/UnrealCloudDDC/Helm/UnrealCloudDDC/templates/configmap-nginx.yaml
+++ b/Engine/Source/Programs/UnrealCloudDDC/Helm/UnrealCloudDDC/templates/configmap-nginx.yaml
@@ -9,7 +9,7 @@ data:
     worker_processes auto;
 
     events {
-      worker_connections  2048;
+      worker_connections {{ .Values.nginx.workerConnections }};
     }
 
     http {
@@ -138,5 +138,20 @@ data:
           return 200; # error_code is ignored
         }
       }
-    }
+
+      server {
+        listen 18080;
+
+        # ensures the version information can be retrieved
+        server_tokens on;
+        
+        location /nginx_status {
+          stub_status on;
+        }
+
+        location / {
+          return 404;
+        }
+      }
+    }  
 {{ end }}
diff --git a/Engine/Source/Programs/UnrealCloudDDC/Helm/UnrealCloudDDC/templates/deployment.yaml b/Engine/Source/Programs/UnrealCloudDDC/Helm/UnrealCloudDDC/templates/deployment.yaml
--- a/Engine/Source/Programs/UnrealCloudDDC/Helm/UnrealCloudDDC/templates/deployment.yaml
+++ b/Engine/Source/Programs/UnrealCloudDDC/Helm/UnrealCloudDDC/templates/deployment.yaml
@@ -136,6 +136,7 @@ spec:
           containerPort: 8091
         - name: nginx-corp-h2
           containerPort: 8092
+        - containerPort: 18080
         {{- if .Values.nginx.livenessProbe.enabled }}
         livenessProbe:
           httpGet:
@@ -162,7 +163,9 @@ spec:
           - name: nginx
             mountPath: /nginx
           - name: nginx-config-volume
-            mountPath: /etc/nginx
+            mountPath: /etc/nginx/nginx.conf
+            subPath: nginx.conf
+            readOnly: true
         {{- if or .Values.persistence.volume .Values.persistence.enabled }}
           - name: data
             mountPath: {{ .Values.persistence.mountPath }}
diff --git a/Engine/Source/Programs/UnrealCloudDDC/Helm/UnrealCloudDDC/values.yaml b/Engine/Source/Programs/UnrealCloudDDC/Helm/UnrealCloudDDC/values.yaml
--- a/Engine/Source/Programs/UnrealCloudDDC/Helm/UnrealCloudDDC/values.yaml
+++ b/Engine/Source/Programs/UnrealCloudDDC/Helm/UnrealCloudDDC/values.yaml
@@ -319,6 +319,9 @@ nginx:
   ## Set this to use Unix domain sockets rather then ports for inter process communication
   useDomainSockets: false
 
+  ## Control the number of worker connections used for nginx
+  workerConnections: 1024
+
   ## Liveness and readiness probes, defaults to corresponding endpoints
   ##
   livenessProbe:
