Unreal Cloud DDC - Delete domain socket if it exists on startup, resolves issue when a container restarts within the pod and getting address in use exceptions.

diff --git a/Engine/Source/Programs/UnrealCloudDDC/Jupiter.Common/BaseProgram.cs b/Engine/Source/Programs/UnrealCloudDDC/Jupiter.Common/BaseProgram.cs
--- a/Engine/Source/Programs/UnrealCloudDDC/Jupiter.Common/BaseProgram.cs
+++ b/Engine/Source/Programs/UnrealCloudDDC/Jupiter.Common/BaseProgram.cs
@@ -56,14 +56,20 @@ namespace Jupiter
 
             try
             {
-                Log.Information("Creating ASPNET Host");
-                IHost host = CreateHostBuilder(args).Build();
-                host.Start();
-
                 JupiterSettings settings = new JupiterSettings();
                 Configuration.GetSection("Jupiter").Bind(settings);
                 string socketsRoot = settings.DomainSocketsRoot;
 
+                if (settings.UseDomainSockets)
+                {
+                    File.Delete(Path.Combine(socketsRoot, "jupiter-http.sock"));
+                    File.Delete(Path.Combine(socketsRoot, "jupiter-http2.sock"));
+                }
+
+                Log.Information("Creating ASPNET Host");
+                IHost host = CreateHostBuilder(args).Build();
+                host.Start();
+
                 if (settings.ChmodDomainSockets)
                 {
                     FileUtils.SetFileMode_Linux(Path.Combine(socketsRoot, "jupiter-http.sock"), 766);
