Unreal Cloud DDC - Added upstream configuration for nginx as well as overriding the number of file limits, lastly you can know enable access logging if you want to.
Also added ability to override the base services targetPort like we can for the extra services.

diff --git a/Engine/Source/Programs/UnrealCloudDDC/Helm/UnrealCloudDDC/templates/configmap-nginx.yaml b/Engine/Source/Programs/UnrealCloudDDC/Helm/UnrealCloudDDC/templates/configmap-nginx.yaml
--- a/Engine/Source/Programs/UnrealCloudDDC/Helm/UnrealCloudDDC/templates/configmap-nginx.yaml
+++ b/Engine/Source/Programs/UnrealCloudDDC/Helm/UnrealCloudDDC/templates/configmap-nginx.yaml
@@ -7,18 +7,31 @@ metadata:
 data:
   nginx.conf: |
     worker_processes auto;
-
+    worker_rlimit_nofile {{ mul .Values.nginx.workerConnections 2 }};
     events {
       worker_connections {{ .Values.nginx.workerConnections }};
     }
 
     http {
+      {{ if .Values.nginx.accessLogging -}}
+      log_format jupiter '$remote_addr - $remote_user [$time_local] '
+                         '"$request" $status $body_bytes_sent rt=$request_time '
+                         '"$http_referer" "$http_user_agent"';
+      access_log /var/log/nginx/access.log jupiter;
+      {{- else -}}
       access_log off;
+      {{- end }}
       client_max_body_size 0;
       server_tokens off;
       sendfile on;
       tcp_nopush on;
 
+      upstream jupiter_backend {
+        zone upstreams 64K;
+        server {{ if .Values.nginx.useDomainSockets -}} unix:/nginx/jupiter-http.sock {{- else -}} localhost:80 {{- end }} max_fails=1 fail_timeout=2s;
+        keepalive 2;
+      }
+
       # http 1.1 public port
       server {
         listen 81 backlog=4096;
@@ -30,12 +43,14 @@ data:
         location /api {
           proxy_set_header X-Jupiter-XAccel-Supported true;
           proxy_set_header X-Jupiter-Port Public;
-          proxy_pass {{ if .Values.nginx.useDomainSockets -}} http://unix:/nginx/jupiter-http.sock {{- else -}} http://localhost:80 {{- end }};
+          proxy_pass http://jupiter_backend;
+          proxy_next_upstream error timeout http_500;
         }
         location /health {
           proxy_set_header X-Jupiter-Port Public;
-          proxy_pass {{ if .Values.nginx.useDomainSockets -}} http://unix:/nginx/jupiter-http.sock {{- else -}} http://localhost:80 {{- end }};
-        }
+          proxy_pass http://jupiter_backend;
+          proxy_next_upstream error timeout http_500;
+         }
         location /nginx-blobs {
           internal;
           alias {{ printf "%s/Blobs" .Values.persistence.mountPath }};
@@ -60,11 +75,13 @@ data:
         location /api {
           proxy_set_header X-Jupiter-XAccel-Supported true;
           proxy_set_header X-Jupiter-Port Corp;
-          proxy_pass {{ if .Values.nginx.useDomainSockets -}} http://unix:/nginx/jupiter-http.sock {{- else -}} http://localhost:80 {{- end }};
+          proxy_pass http://jupiter_backend;
+          proxy_next_upstream error timeout http_500;  
         }
         location /health {
           proxy_set_header X-Jupiter-Port Corp;
-          proxy_pass {{ if .Values.nginx.useDomainSockets -}} http://unix:/nginx/jupiter-http.sock {{- else -}} http://localhost:80 {{- end }};
+          proxy_pass http://jupiter_backend;
+          proxy_next_upstream error timeout http_500;
         }
         location /nginx-blobs {
           internal;
@@ -90,11 +107,13 @@ data:
         location /api {
           proxy_set_header X-Jupiter-XAccel-Supported true;
           proxy_set_header X-Jupiter-Port Public;
-          proxy_pass {{ if .Values.nginx.useDomainSockets -}} http://unix:/nginx/jupiter-http.sock {{- else -}} http://localhost:80 {{- end }};
+          proxy_pass http://jupiter_backend;
+          proxy_next_upstream error timeout http_500;
         }
         location /health {
           proxy_set_header X-Jupiter-Port Public;
-          proxy_pass {{ if .Values.nginx.useDomainSockets -}} http://unix:/nginx/jupiter-http.sock {{- else -}} http://localhost:80 {{- end }};
+          proxy_pass http://jupiter_backend;
+          proxy_next_upstream error timeout http_500;
         }
         location /nginx-blobs {
           internal;
@@ -120,11 +139,13 @@ data:
         location /api {
           proxy_set_header X-Jupiter-XAccel-Supported true;
           proxy_set_header X-Jupiter-Port Corp;
-          proxy_pass {{ if .Values.nginx.useDomainSockets -}} http://unix:/nginx/jupiter-http.sock {{- else -}} http://localhost:80 {{- end }};
+          proxy_pass http://jupiter_backend;
+          proxy_next_upstream error timeout http_500;
         }
         location /health {
           proxy_set_header X-Jupiter-Port Corp;
-          proxy_pass {{ if .Values.nginx.useDomainSockets -}} http://unix:/nginx/jupiter-http.sock {{- else -}} http://localhost:80 {{- end }};
+          proxy_pass http://jupiter_backend;
+          proxy_next_upstream error timeout http_500;
         }
         location /nginx-blobs {
             internal;
diff --git a/Engine/Source/Programs/UnrealCloudDDC/Helm/UnrealCloudDDC/templates/service.yaml b/Engine/Source/Programs/UnrealCloudDDC/Helm/UnrealCloudDDC/templates/service.yaml
--- a/Engine/Source/Programs/UnrealCloudDDC/Helm/UnrealCloudDDC/templates/service.yaml
+++ b/Engine/Source/Programs/UnrealCloudDDC/Helm/UnrealCloudDDC/templates/service.yaml
@@ -15,7 +15,7 @@ spec:
   ports:
     - name: http
       port: {{ .Values.service.port }}
-      targetPort: http
+      targetPort: {{ default "http" .Values.service.targetPort }}
       {{- if and (or (eq .Values.service.type "NodePort") (eq .Values.service.type "LoadBalancer")) (not (empty .Values.service.nodePort)) }}
       nodePort: {{ .Values.service.nodePort }}
       {{- else if eq .Values.service.type "ClusterIP" }}
diff --git a/Engine/Source/Programs/UnrealCloudDDC/Helm/UnrealCloudDDC/values.yaml b/Engine/Source/Programs/UnrealCloudDDC/Helm/UnrealCloudDDC/values.yaml
--- a/Engine/Source/Programs/UnrealCloudDDC/Helm/UnrealCloudDDC/values.yaml
+++ b/Engine/Source/Programs/UnrealCloudDDC/Helm/UnrealCloudDDC/values.yaml
@@ -319,6 +319,9 @@ nginx:
   ## Set this to use Unix domain sockets rather then ports for inter process communication
   useDomainSockets: false
 
+  ## Set to enable access log output, can be useful for debugging but will generate a lot of logs in prod
+  accessLogging: false
+
   ## Control the number of worker connections used for nginx
   workerConnections: 1024
 
