EpicGames.Horde: Remove legacy storage client interface.

diff --git a/Engine/Source/Programs/Horde/Horde.Agent.Tests/AwsLambdaListenerCommandTest.cs b/Engine/Source/Programs/Horde/Horde.Agent.Tests/AwsLambdaListenerCommandTest.cs
deleted file mode 100644
diff --git a/Engine/Source/Programs/Horde/Horde.Agent/AgentSettings.cs b/Engine/Source/Programs/Horde/Horde.Agent/AgentSettings.cs
--- a/Engine/Source/Programs/Horde/Horde.Agent/AgentSettings.cs
+++ b/Engine/Source/Programs/Horde/Horde.Agent/AgentSettings.cs
@@ -64,11 +64,6 @@ namespace Horde.Agent
 		/// </summary>
 		public List<string> Thumbprints { get; } = new List<string>();
 
-		/// <summary>
-		/// Storage settings for using this server
-		/// </summary>
-		public StorageOptions Storage { get; set; } = new StorageOptions();
-
 		/// <summary>
 		/// Checks whether the given certificate thumbprint should be trusted
 		/// </summary>
diff --git a/Engine/Source/Programs/Horde/Horde.Agent/Commands/Compute/AwsLambdaListenerCommand.cs b/Engine/Source/Programs/Horde/Horde.Agent/Commands/Compute/AwsLambdaListenerCommand.cs
deleted file mode 100644
diff --git a/Engine/Source/Programs/Horde/Horde.Agent/Horde.Agent.csproj b/Engine/Source/Programs/Horde/Horde.Agent/Horde.Agent.csproj
--- a/Engine/Source/Programs/Horde/Horde.Agent/Horde.Agent.csproj
+++ b/Engine/Source/Programs/Horde/Horde.Agent/Horde.Agent.csproj
@@ -138,7 +138,6 @@
     <Protobuf Include="..\Horde.Build\Protos\horde\log_rpc.proto" ProtoRoot="..\Horde.Build\Protos" GrpcServices="Client" Access="Public" />
     <Protobuf Include="..\Horde.Build\Protos\horde\log_rpc_messages.proto" ProtoRoot="..\Horde.Build\Protos" Access="Public" />
     <Protobuf Include="..\Horde.Build\Protos\horde\common\agent_status.proto" ProtoRoot="..\Horde.Build\Protos" Access="Public" />
-    <Protobuf Include="..\Horde.Build\Protos\horde\common\bucket_id.proto" ProtoRoot="..\Horde.Build\Protos" Access="Public" />
     <Protobuf Include="..\Horde.Build\Protos\horde\common\event_severity.proto" ProtoRoot="..\Horde.Build\Protos" Access="Public" />
     <Protobuf Include="..\Horde.Build\Protos\horde\common\io_hash.proto" ProtoRoot="..\Horde.Build\Protos" Access="Public" />
     <Protobuf Include="..\Horde.Build\Protos\horde\common\jobstep_outcome.proto" ProtoRoot="..\Horde.Build\Protos" Access="Public" />
@@ -148,8 +147,6 @@
     <Protobuf Include="..\Horde.Build\Protos\horde\common\lease_state.proto" ProtoRoot="..\Horde.Build\Protos" Access="Public" />
     <Protobuf Include="..\Horde.Build\Protos\horde\common\namespace_id.proto" ProtoRoot="..\Horde.Build\Protos" Access="Public" />
     <Protobuf Include="..\Horde.Build\Protos\horde\common\priority.proto" ProtoRoot="..\Horde.Build\Protos" Access="Public" />
-    <Protobuf Include="..\Horde.Build\Protos\horde\common\qualified_ref_id.proto" ProtoRoot="..\Horde.Build\Protos" Access="Public" />
-    <Protobuf Include="..\Horde.Build\Protos\horde\common\ref_id.proto" ProtoRoot="..\Horde.Build\Protos" Access="Public" />
     <Protobuf Include="..\Horde.Build\Protos\horde\messages\agent_capabilities.proto" ProtoRoot="..\Horde.Build\Protos" Access="Public" />
     <Protobuf Include="..\Horde.Build\Protos\horde\messages\agent_requirements.proto" ProtoRoot="..\Horde.Build\Protos" Access="Public" />
     <Protobuf Include="..\Horde.Build\Protos\horde\messages\agent_workspace.proto" ProtoRoot="..\Horde.Build\Protos" Access="Public" />
diff --git a/Engine/Source/Programs/Horde/Horde.Agent/Program.cs b/Engine/Source/Programs/Horde/Horde.Agent/Program.cs
--- a/Engine/Source/Programs/Horde/Horde.Agent/Program.cs
+++ b/Engine/Source/Programs/Horde/Horde.Agent/Program.cs
@@ -180,8 +180,6 @@ namespace Horde.Agent
 					return builder.WaitAndRetryAsync(new[] { TimeSpan.FromSeconds(1), TimeSpan.FromSeconds(5), TimeSpan.FromSeconds(10) });
 				});
 
-			services.AddHordeStorage(settings => configSection.GetCurrentServerProfile().GetSection(nameof(serverProfile.Storage)).Bind(settings));
-
 			services.AddSingleton<GrpcService>();
 
 			services.AddSingleton<IJobExecutorFactory, PerforceExecutorFactory>();
diff --git a/Engine/Source/Programs/Horde/Horde.Build.Tests/StorageTests.cs b/Engine/Source/Programs/Horde/Horde.Build.Tests/StorageTests.cs
deleted file mode 100644
diff --git a/Engine/Source/Programs/Horde/Horde.Build.Tests/TestSetup.cs b/Engine/Source/Programs/Horde/Horde.Build.Tests/TestSetup.cs
--- a/Engine/Source/Programs/Horde/Horde.Build.Tests/TestSetup.cs
+++ b/Engine/Source/Programs/Horde/Horde.Build.Tests/TestSetup.cs
@@ -25,7 +25,6 @@ using Horde.Build.Logs.Builder;
 using Horde.Build.Notifications;
 using Horde.Build.Server;
 using Horde.Build.Storage;
-using Horde.Build.Storage.Services;
 using Horde.Build.Tests.Stubs.Collections;
 using Horde.Build.Tests.Stubs.Services;
 using Horde.Build.Tools;
@@ -266,9 +265,6 @@ namespace Horde.Build.Tests
 
 			services.AddSingleton<IStorageBackend<PersistentLogStorage>>(sp => new MemoryStorageBackend().ForType<PersistentLogStorage>());
 			services.AddSingleton<IStorageBackend<ArtifactCollectionV1>>(sp => new MemoryStorageBackend().ForType<ArtifactCollectionV1>());
-			services.AddSingleton<IStorageBackend<BasicStorageClient>>(sp => new MemoryStorageBackend().ForType<BasicStorageClient>());
-
-			services.AddSingleton<ILegacyStorageClient, BasicStorageClient>();
 
 			services.AddSingleton<StorageService>();
 
diff --git a/Engine/Source/Programs/Horde/Horde.Build/Horde.Build.csproj b/Engine/Source/Programs/Horde/Horde.Build/Horde.Build.csproj
--- a/Engine/Source/Programs/Horde/Horde.Build/Horde.Build.csproj
+++ b/Engine/Source/Programs/Horde/Horde.Build/Horde.Build.csproj
@@ -25,6 +25,7 @@
     <PackageReference Include="AWSSDK.Extensions.NETCore.Setup" Version="3.7.2" />
     <PackageReference Include="AWSSDK.S3" Version="3.7.100" />
     <PackageReference Include="AWSSDK.SecurityToken" Version="3.7.100" />
+    <PackageReference Include="Crc32.NET" Version="1.2.0" />
     <PackageReference Include="Datadog.Trace.OpenTracing" Version="2.14.0" />
     <PackageReference Include="DogStatsD-CSharp-Client" Version="7.0.1" />
     <PackageReference Include="Grpc.AspNetCore" Version="2.49.0" />
@@ -79,7 +80,6 @@
     <Protobuf Include="Protos\horde\log_rpc.proto" ProtoRoot="Protos" GrpcServices="Server" />
     <Protobuf Include="Protos\horde\log_rpc_messages.proto" ProtoRoot="Protos" />
     <Protobuf Include="Protos\horde\common\agent_status.proto" ProtoRoot="Protos" />
-    <Protobuf Include="Protos\horde\common\bucket_id.proto" ProtoRoot="Protos" />
     <Protobuf Include="Protos\horde\common\event_severity.proto" ProtoRoot="Protos" />
     <Protobuf Include="Protos\horde\common\io_hash.proto" ProtoRoot="Protos" />
     <Protobuf Include="Protos\horde\common\jobstep_outcome.proto" ProtoRoot="Protos" />
@@ -89,8 +89,6 @@
     <Protobuf Include="Protos\horde\common\lease_state.proto" ProtoRoot="Protos" />
     <Protobuf Include="Protos\horde\common\namespace_id.proto" ProtoRoot="Protos" />
     <Protobuf Include="Protos\horde\common\priority.proto" ProtoRoot="Protos" />
-    <Protobuf Include="Protos\horde\common\qualified_ref_id.proto" ProtoRoot="Protos" />
-    <Protobuf Include="Protos\horde\common\ref_id.proto" ProtoRoot="Protos" />
     <Protobuf Include="Protos\horde\messages\agent_capabilities.proto" ProtoRoot="Protos" />
     <Protobuf Include="Protos\horde\messages\agent_requirements.proto" ProtoRoot="Protos" />
     <Protobuf Include="Protos\horde\messages\agent_workspace.proto" ProtoRoot="Protos" />
@@ -110,6 +108,7 @@
     <ProjectReference Include="..\..\Shared\EpicGames.Core\EpicGames.Core.csproj" />
     <ProjectReference Include="..\..\Shared\EpicGames.Horde\EpicGames.Horde.csproj" />
     <ProjectReference Include="..\..\Shared\EpicGames.IoHash\EpicGames.IoHash.csproj" />
+    <ProjectReference Include="..\..\Shared\EpicGames.Oodle\EpicGames.Oodle.csproj" />
     <ProjectReference Include="..\..\Shared\EpicGames.Perforce.Managed\EpicGames.Perforce.Managed.csproj" />
     <ProjectReference Include="..\..\Shared\EpicGames.Perforce\EpicGames.Perforce.csproj" />
     <ProjectReference Include="..\..\Shared\EpicGames.Redis\EpicGames.Redis.csproj" />
diff --git a/Engine/Source/Programs/Horde/Horde.Build/Protos/horde/common/bucket_id.cs b/Engine/Source/Programs/Horde/Horde.Build/Protos/horde/common/bucket_id.cs
deleted file mode 100644
diff --git a/Engine/Source/Programs/Horde/Horde.Build/Protos/horde/common/qualified_ref_id.cs b/Engine/Source/Programs/Horde/Horde.Build/Protos/horde/common/qualified_ref_id.cs
deleted file mode 100644
diff --git a/Engine/Source/Programs/Horde/Horde.Build/Protos/horde/common/ref_id.cs b/Engine/Source/Programs/Horde/Horde.Build/Protos/horde/common/ref_id.cs
deleted file mode 100644
diff --git a/Engine/Source/Programs/Horde/Horde.Build/ServerSettings.cs b/Engine/Source/Programs/Horde/Horde.Build/ServerSettings.cs
--- a/Engine/Source/Programs/Horde/Horde.Build/ServerSettings.cs
+++ b/Engine/Source/Programs/Horde/Horde.Build/ServerSettings.cs
@@ -654,11 +654,6 @@ namespace Horde.Build
 		/// </summary>
 		public List<PerforceConnectionSettings> Perforce { get; set; } = new List<PerforceConnectionSettings>();
 
-		/// <summary>
-		/// Settings for the storage service
-		/// </summary>
-		public StorageOptions? Storage { get; set; }
-
 		/// <summary>
 		/// Whether to use the local Perforce environment
 		/// </summary>
diff --git a/Engine/Source/Programs/Horde/Horde.Build/Startup.cs b/Engine/Source/Programs/Horde/Horde.Build/Startup.cs
--- a/Engine/Source/Programs/Horde/Horde.Build/Startup.cs
+++ b/Engine/Source/Programs/Horde/Horde.Build/Startup.cs
@@ -475,8 +475,6 @@ namespace Horde.Build
 			services.AddSingleton(sp => CreateStorageBackend(sp, settings.LogStorage).ForType<PersistentLogStorage>());
 			services.AddSingleton(sp => CreateStorageBackend(sp, settings.ArtifactStorage).ForType<ArtifactCollectionV1>());
 
-			services.AddHordeStorage(settings => configSection.GetSection("Storage").Bind(settings));
-
 			if (settings.WithAws)
 			{
 				AWSOptions awsOptions = Configuration.GetAWSOptions();
@@ -881,21 +879,6 @@ namespace Horde.Build
 			}
 		}
 
-		public sealed class RefIdBsonSerializer : SerializerBase<RefId>
-		{
-			/// <inheritdoc/>
-			public override RefId Deserialize(BsonDeserializationContext context, BsonDeserializationArgs args)
-			{
-				return new RefId(IoHash.Parse(context.Reader.ReadString()));
-			}
-
-			/// <inheritdoc/>
-			public override void Serialize(BsonSerializationContext context, BsonSerializationArgs args, RefId value)
-			{
-				context.Writer.WriteString(value.ToString());
-			}
-		}
-
 		public sealed class RefNameBsonSerializer : SerializerBase<RefName>
 		{
 			/// <inheritdoc/>
@@ -957,7 +940,6 @@ namespace Horde.Build
 				BsonSerializer.RegisterSerializer(new HostIdBsonSerializer());
 				BsonSerializer.RegisterSerializer(new BlobIdBsonSerializer());
 				BsonSerializer.RegisterSerializer(new BlobLocatorBsonSerializer());
-				BsonSerializer.RegisterSerializer(new RefIdBsonSerializer());
 				BsonSerializer.RegisterSerializer(new RefNameBsonSerializer());
 				BsonSerializer.RegisterSerializer(new IoHashBsonSerializer());
 				BsonSerializer.RegisterSerializer(new NamespaceIdBsonSerializer());
diff --git a/Engine/Source/Programs/Horde/Horde.Build/Storage/Services/BasicStorageClient.cs b/Engine/Source/Programs/Horde/Horde.Build/Storage/Services/BasicStorageClient.cs
deleted file mode 100644
diff --git a/Engine/Source/Programs/Horde/Horde.Build/Utilities/ObjectSet.cs b/Engine/Source/Programs/Horde/Horde.Build/Utilities/ObjectSet.cs
deleted file mode 100644
diff --git a/Engine/Source/Programs/Shared/EpicGames.Horde.Tests/StorageClientTest.cs b/Engine/Source/Programs/Shared/EpicGames.Horde.Tests/StorageClientTest.cs
deleted file mode 100644
diff --git a/Engine/Source/Programs/Shared/EpicGames.Horde/Storage/Backends/HttpStorageClient.cs b/Engine/Source/Programs/Shared/EpicGames.Horde/Storage/Backends/HttpStorageClient.cs
--- a/Engine/Source/Programs/Shared/EpicGames.Horde/Storage/Backends/HttpStorageClient.cs
+++ b/Engine/Source/Programs/Shared/EpicGames.Horde/Storage/Backends/HttpStorageClient.cs
@@ -8,14 +8,10 @@ using System.Net.Http;
 using System.Net.Http.Headers;
 using System.Net.Http.Json;
 using System.Runtime.CompilerServices;
-using System.Runtime.InteropServices;
 using System.Threading;
 using System.Threading.Tasks;
 using System.Web;
-using System.Xml.Linq;
 using EpicGames.Core;
-using EpicGames.Horde.Storage;
-using Microsoft.Extensions.Caching.Memory;
 using Microsoft.Extensions.Logging;
 
 namespace EpicGames.Horde.Storage.Backends
@@ -150,7 +146,7 @@ namespace EpicGames.Horde.Storage.Backends
 							if (!uploadResponse.IsSuccessStatusCode)
 							{
 								string body = await uploadResponse.Content.ReadAsStringAsync(cancellationToken);
-								throw new StorageException($"Unable to upload data to redirected URL: {body}", null);
+								throw new StorageException($"Unable to upload data to redirected URL: {body}");
 							}
 						}
 						_logger.LogDebug("Written {Locator} (using redirect)", redirectResponse.Blob);
@@ -254,7 +250,7 @@ namespace EpicGames.Horde.Storage.Backends
 						else if (!response.IsSuccessStatusCode)
 						{
 							_logger.LogError("Unable to read ref {RefName} (status: {StatusCode}, body: {Body})", name, response.StatusCode, await response.Content.ReadAsStringAsync(cancellationToken));
-							throw new StorageException($"Unable to read ref '{name}'", null);
+							throw new StorageException($"Unable to read ref '{name}'");
 						}
 						else
 						{
diff --git a/Engine/Source/Programs/Shared/EpicGames.Horde/Storage/BucketId.cs b/Engine/Source/Programs/Shared/EpicGames.Horde/Storage/BucketId.cs
deleted file mode 100644
diff --git a/Engine/Source/Programs/Shared/EpicGames.Horde/Storage/ILegacyStorageClient.cs b/Engine/Source/Programs/Shared/EpicGames.Horde/Storage/ILegacyStorageClient.cs
deleted file mode 100644
diff --git a/Engine/Source/Programs/Shared/EpicGames.Horde/Storage/IStorageClient.cs b/Engine/Source/Programs/Shared/EpicGames.Horde/Storage/IStorageClient.cs
--- a/Engine/Source/Programs/Shared/EpicGames.Horde/Storage/IStorageClient.cs
+++ b/Engine/Source/Programs/Shared/EpicGames.Horde/Storage/IStorageClient.cs
@@ -15,10 +15,32 @@ using EpicGames.Core;
 
 namespace EpicGames.Horde.Storage
 {
+	/// <summary>
+	/// Base exception for the storage service
+	/// </summary>
+	public class StorageException : Exception
+	{
+		/// <summary>
+		/// Constructor
+		/// </summary>
+		public StorageException(string message)
+			: base(message)
+		{
+		}
+
+		/// <summary>
+		/// Constructor
+		/// </summary>
+		public StorageException(string message, Exception? innerException)
+			: base(message, innerException)
+		{
+		}
+	}
+
 	/// <summary>
 	/// Exception for a ref not existing
 	/// </summary>
-	public sealed class RefNameNotFoundException : Exception
+	public sealed class RefNameNotFoundException : StorageException
 	{
 		/// <summary>
 		/// Name of the missing ref
diff --git a/Engine/Source/Programs/Shared/EpicGames.Horde/Storage/Impl/FileStorageClient.cs b/Engine/Source/Programs/Shared/EpicGames.Horde/Storage/Impl/FileStorageClient.cs
deleted file mode 100644
diff --git a/Engine/Source/Programs/Shared/EpicGames.Horde/Storage/Impl/HttpStorageClient.cs b/Engine/Source/Programs/Shared/EpicGames.Horde/Storage/Impl/HttpStorageClient.cs
deleted file mode 100644
diff --git a/Engine/Source/Programs/Shared/EpicGames.Horde/Storage/Impl/MemoryStorageClient.cs b/Engine/Source/Programs/Shared/EpicGames.Horde/Storage/Impl/MemoryStorageClient.cs
deleted file mode 100644
diff --git a/Engine/Source/Programs/Shared/EpicGames.Horde/Storage/Impl/TracingStorageClient.cs b/Engine/Source/Programs/Shared/EpicGames.Horde/Storage/Impl/TracingStorageClient.cs
deleted file mode 100644
diff --git a/Engine/Source/Programs/Shared/EpicGames.Horde/Storage/QualifiedRefId.cs b/Engine/Source/Programs/Shared/EpicGames.Horde/Storage/QualifiedRefId.cs
deleted file mode 100644
diff --git a/Engine/Source/Programs/Shared/EpicGames.Horde/Storage/RefId.cs b/Engine/Source/Programs/Shared/EpicGames.Horde/Storage/RefId.cs
deleted file mode 100644
diff --git a/Engine/Source/Programs/Shared/EpicGames.Horde/Storage/StorageExtensions.cs b/Engine/Source/Programs/Shared/EpicGames.Horde/Storage/StorageExtensions.cs
deleted file mode 100644
diff --git a/Engine/Source/Programs/UnrealCloudDDC/Jupiter.Common/BucketId.cs b/Engine/Source/Programs/UnrealCloudDDC/Jupiter.Common/BucketId.cs
new file mode 100644
--- /dev/null
+++ b/Engine/Source/Programs/UnrealCloudDDC/Jupiter.Common/BucketId.cs
@@ -0,0 +1,88 @@
+// Copyright Epic Games, Inc. All Rights Reserved.
+
+using System;
+using System.ComponentModel;
+using System.Globalization;
+using System.Text.Json;
+using System.Text.Json.Serialization;
+using EpicGames.Core;
+using EpicGames.Serialization;
+
+namespace EpicGames.Horde.Storage
+{
+	/// <summary>
+	/// Identifier for a storage bucket
+	/// </summary>
+	[JsonConverter(typeof(BucketIdJsonConverter))]
+	[TypeConverter(typeof(BucketIdTypeConverter))]
+	[CbConverter(typeof(BucketIdCbConverter))]
+	public struct BucketId : IEquatable<BucketId>
+	{
+		/// <summary>
+		/// The text representing this id
+		/// </summary>
+		readonly StringId _inner;
+
+		/// <summary>
+		/// Constructor
+		/// </summary>
+		/// <param name="input">Unique id for the string</param>
+		public BucketId(string input)
+		{
+			_inner = new StringId(input);
+		}
+
+		/// <inheritdoc/>
+		public override bool Equals(object? obj) => obj is BucketId id && _inner.Equals(id._inner);
+
+		/// <inheritdoc/>
+		public override int GetHashCode() => _inner.GetHashCode();
+
+		/// <inheritdoc/>
+		public bool Equals(BucketId other) => _inner.Equals(other._inner);
+
+		/// <inheritdoc/>
+		public override string ToString() => _inner.ToString();
+
+		/// <inheritdoc cref="StringId.op_Equality"/>
+		public static bool operator ==(BucketId left, BucketId right) => left._inner == right._inner;
+
+		/// <inheritdoc cref="StringId.op_Inequality"/>
+		public static bool operator !=(BucketId left, BucketId right) => left._inner != right._inner;
+	}
+
+	/// <summary>
+	/// Type converter for BucketId to and from JSON
+	/// </summary>
+	sealed class BucketIdJsonConverter : JsonConverter<BucketId>
+	{
+		/// <inheritdoc/>
+		public override BucketId Read(ref Utf8JsonReader reader, Type typeToConvert, JsonSerializerOptions options) => new BucketId(reader.GetString() ?? String.Empty);
+
+		/// <inheritdoc/>
+		public override void Write(Utf8JsonWriter writer, BucketId value, JsonSerializerOptions options) => writer.WriteStringValue(value.ToString());
+	}
+
+	/// <summary>
+	/// Type converter from strings to BucketId objects
+	/// </summary>
+	sealed class BucketIdTypeConverter : TypeConverter
+	{
+		/// <inheritdoc/>
+		public override bool CanConvertFrom(ITypeDescriptorContext? context, Type? sourceType) => sourceType == typeof(string);
+
+		/// <inheritdoc/>
+		public override object ConvertFrom(ITypeDescriptorContext? context, CultureInfo? culture, object? value) => new BucketId((string)value!);
+	}
+
+	sealed class BucketIdCbConverter : CbConverterBase<BucketId>
+	{
+		public override BucketId Read(CbField field) => new BucketId(field.AsString());
+
+		/// <inheritdoc/>
+		public override void Write(CbWriter writer, BucketId value) => writer.WriteStringValue(value.ToString());
+
+		/// <inheritdoc/>
+		public override void WriteNamed(CbWriter writer, Utf8String name, BucketId value) => writer.WriteString(name, value.ToString());
+	}
+}
