Unreal Cloud DDC - Process each attachment of a package in parallel when fetching a ref as a cb package. Also added some trace scopes to see how much time we spend on this.

diff --git a/Engine/Source/Programs/UnrealCloudDDC/Jupiter/Controllers/ReferencesController.cs b/Engine/Source/Programs/UnrealCloudDDC/Jupiter/Controllers/ReferencesController.cs
--- a/Engine/Source/Programs/UnrealCloudDDC/Jupiter/Controllers/ReferencesController.cs
+++ b/Engine/Source/Programs/UnrealCloudDDC/Jupiter/Controllers/ReferencesController.cs
@@ -222,7 +222,7 @@ namespace Jupiter.Controllers
                     {
                         byte[] blobMemory;
                         {
-                            using TelemetrySpan scope = _tracer.StartActiveSpan("json.readblob").SetAttribute("operation.name", "authorize");
+                            using TelemetrySpan scope = _tracer.StartActiveSpan("json.readblob").SetAttribute("operation.name", "json.readblob");
                             blobMemory = await blob.Stream.ToByteArray();
                         }
                         CbObject cb = new CbObject(blobMemory);
@@ -234,6 +234,7 @@ namespace Jupiter.Controllers
                     }
                     case CustomMediaTypeNames.UnrealCompactBinaryPackage:
                     {
+                        using TelemetrySpan packageScope = _tracer.StartActiveSpan("cbpackage.fetch").SetAttribute("operation.name", "cbpackage.fetch");
                         byte[] blobMemory = await blob.Stream.ToByteArray();
                         CbObject cb = new CbObject(blobMemory);
 
@@ -242,7 +243,7 @@ namespace Jupiter.Controllers
                         using CbPackageBuilder writer = new CbPackageBuilder();
                         writer.AddAttachment(objectRecord.BlobIdentifier.AsIoHash(), CbPackageAttachmentFlags.IsObject, blobMemory);
 
-                        await foreach (Attachment attachment in attachments)
+                        await Parallel.ForEachAsync(attachments, async (attachment, token) =>
                         {
                             IoHash attachmentHash = attachment.AsIoHash();
                             CbPackageAttachmentFlags flags = 0;
@@ -287,10 +288,18 @@ namespace Jupiter.Controllers
                             catch (Exception e)
                             {
                                 (CbObject errorObject, HttpStatusCode _) = ToErrorResult(e);
-                                writer.AddAttachment(attachmentHash, CbPackageAttachmentFlags.IsError | CbPackageAttachmentFlags.IsObject, errorObject.GetView().ToArray());
+                                writer.AddAttachment(attachmentHash,
+                                    CbPackageAttachmentFlags.IsError | CbPackageAttachmentFlags.IsObject,
+                                    errorObject.GetView().ToArray());
                             }
+                        });
+
+                        byte[] packageBytes;
+                        {
+                            using TelemetrySpan _ = _tracer.StartActiveSpan("cbpackage.buffer").SetAttribute("operation.name", "cbpackage.buffer");
+                            packageBytes = await writer.ToByteArray();
                         }
-                        await using BlobContents contents = new BlobContents(await writer.ToByteArray());
+                        await using BlobContents contents = new BlobContents(packageBytes);
                         await WriteBody(contents, CustomMediaTypeNames.UnrealCompactBinaryPackage);
                         break;
                     }
