Unreal Cloud DDC - Add response type to tracing scope for get requests, to be able to find only requests of a certain type

diff --git a/Engine/Source/Programs/UnrealCloudDDC/Jupiter/Controllers/ReferencesController.cs b/Engine/Source/Programs/UnrealCloudDDC/Jupiter/Controllers/ReferencesController.cs
--- a/Engine/Source/Programs/UnrealCloudDDC/Jupiter/Controllers/ReferencesController.cs
+++ b/Engine/Source/Programs/UnrealCloudDDC/Jupiter/Controllers/ReferencesController.cs
@@ -162,7 +162,7 @@ namespace Jupiter.Controllers
                 }
 
                 string responseType = _formatResolver.GetResponseType(Request, format, CustomMediaTypeNames.UnrealCompactBinary);
-
+                Tracer.CurrentSpan.SetAttribute("response-type", responseType);
                 switch (responseType)
                 {
                     case CustomMediaTypeNames.UnrealCompactBinary:
