Unreal Cloud DDC - Prepare release 0.3.0

diff --git a/Engine/Source/Programs/UnrealCloudDDC/Helm/UnrealCloudDDC/Chart.yaml b/Engine/Source/Programs/UnrealCloudDDC/Helm/UnrealCloudDDC/Chart.yaml
--- a/Engine/Source/Programs/UnrealCloudDDC/Helm/UnrealCloudDDC/Chart.yaml
+++ b/Engine/Source/Programs/UnrealCloudDDC/Helm/UnrealCloudDDC/Chart.yaml
@@ -4,7 +4,7 @@ description: A Helm chart for Kubernetes deploying UnrealCloudDDC
 
 type: application
 
-version: 0.2.0
+version: 0.3.0
 
 icon: https://devtools-misc.s3.amazonaws.com/ue4logo.png
 home: https://unrealengine.com
diff --git a/Engine/Source/Programs/UnrealCloudDDC/Helm/UnrealCloudDDC/version.yaml b/Engine/Source/Programs/UnrealCloudDDC/Helm/UnrealCloudDDC/version.yaml
--- a/Engine/Source/Programs/UnrealCloudDDC/Helm/UnrealCloudDDC/version.yaml
+++ b/Engine/Source/Programs/UnrealCloudDDC/Helm/UnrealCloudDDC/version.yaml
@@ -1 +1 @@
-version: 0.2.0
\ No newline at end of file
+version: 0.3.0
\ No newline at end of file
diff --git a/Engine/Source/Programs/UnrealCloudDDC/version.yaml b/Engine/Source/Programs/UnrealCloudDDC/version.yaml
--- a/Engine/Source/Programs/UnrealCloudDDC/version.yaml
+++ b/Engine/Source/Programs/UnrealCloudDDC/version.yaml
@@ -1 +1 @@
-version: 0.2.0
\ No newline at end of file
+version: 0.3.0
\ No newline at end of file
