Unreal Cloud DDC - Made sure nginx uses keep alive against the backend servers

diff --git a/Engine/Source/Programs/UnrealCloudDDC/Helm/UnrealCloudDDC/templates/configmap-nginx.yaml b/Engine/Source/Programs/UnrealCloudDDC/Helm/UnrealCloudDDC/templates/configmap-nginx.yaml
--- a/Engine/Source/Programs/UnrealCloudDDC/Helm/UnrealCloudDDC/templates/configmap-nginx.yaml
+++ b/Engine/Source/Programs/UnrealCloudDDC/Helm/UnrealCloudDDC/templates/configmap-nginx.yaml
@@ -29,7 +29,7 @@ data:
       upstream jupiter_backend {
         zone upstreams 64K;
         server {{ if .Values.nginx.useDomainSockets -}} unix:/nginx/jupiter-http.sock {{- else -}} localhost:80 {{- end }} max_fails=1 fail_timeout=2s;
-        keepalive 2;
+        keepalive 64;
       }
 
       # http 1.1 public port
@@ -43,11 +43,13 @@ data:
         location /api {
           proxy_set_header X-Jupiter-XAccel-Supported true;
           proxy_set_header X-Jupiter-Port Public;
+          proxy_set_header Connection "";
           proxy_pass http://jupiter_backend;
           proxy_next_upstream error timeout http_500;
         }
         location /health {
           proxy_set_header X-Jupiter-Port Public;
+          proxy_set_header Connection "";
           proxy_pass http://jupiter_backend;
           proxy_next_upstream error timeout http_500;
          }
@@ -75,11 +77,13 @@ data:
         location /api {
           proxy_set_header X-Jupiter-XAccel-Supported true;
           proxy_set_header X-Jupiter-Port Corp;
+          proxy_set_header Connection "";
           proxy_pass http://jupiter_backend;
           proxy_next_upstream error timeout http_500;  
         }
         location /health {
           proxy_set_header X-Jupiter-Port Corp;
+          proxy_set_header Connection "";
           proxy_pass http://jupiter_backend;
           proxy_next_upstream error timeout http_500;
         }
@@ -107,11 +111,13 @@ data:
         location /api {
           proxy_set_header X-Jupiter-XAccel-Supported true;
           proxy_set_header X-Jupiter-Port Public;
+          proxy_set_header Connection "";
           proxy_pass http://jupiter_backend;
           proxy_next_upstream error timeout http_500;
         }
         location /health {
           proxy_set_header X-Jupiter-Port Public;
+          proxy_set_header Connection "";
           proxy_pass http://jupiter_backend;
           proxy_next_upstream error timeout http_500;
         }
@@ -139,11 +145,13 @@ data:
         location /api {
           proxy_set_header X-Jupiter-XAccel-Supported true;
           proxy_set_header X-Jupiter-Port Corp;
+          proxy_set_header Connection "";
           proxy_pass http://jupiter_backend;
           proxy_next_upstream error timeout http_500;
         }
         location /health {
           proxy_set_header X-Jupiter-Port Corp;
+          proxy_set_header Connection "";
           proxy_pass http://jupiter_backend;
           proxy_next_upstream error timeout http_500;
         }
