Unreal Cloud DDC - Fix broken tests when legacy buckets table is missing

diff --git a/Engine/Source/Programs/UnrealCloudDDC/Jupiter/Implementation/Objects/ScyllaReferencesStore.cs b/Engine/Source/Programs/UnrealCloudDDC/Jupiter/Implementation/Objects/ScyllaReferencesStore.cs
--- a/Engine/Source/Programs/UnrealCloudDDC/Jupiter/Implementation/Objects/ScyllaReferencesStore.cs
+++ b/Engine/Source/Programs/UnrealCloudDDC/Jupiter/Implementation/Objects/ScyllaReferencesStore.cs
@@ -69,6 +69,16 @@ namespace Jupiter.Implementation
             );"
             ));
 
+            if (_settings.CurrentValue.ListObjectsFromOldNamespaceTable)
+            {
+                _session.Execute(new SimpleStatement(@"CREATE TABLE IF NOT EXISTS buckets (
+                    namespace text, 
+                    bucket set<text>, 
+                    PRIMARY KEY (namespace)
+                );"
+                ));
+            }
+
             // BYPASS CACHE is a scylla specific extension to disable populating the cache, should be ignored by other cassandra dbs
             string cqlOptions = scyllaSessionManager.IsScylla ? "BYPASS CACHE" : "";
             _getObjectsStatement = _session.Prepare($"SELECT namespace, bucket, name, last_access_time FROM objects ALLOW FILTERING {cqlOptions}");
