Unreal Cloud DDC - Retry replication of op if there is a socket reset exception - reduces noise from our error reporting as this exception happens a few times a day

diff --git a/Engine/Source/Programs/UnrealCloudDDC/Jupiter/Implementation/Replication/RefsReplicator.cs b/Engine/Source/Programs/UnrealCloudDDC/Jupiter/Implementation/Replication/RefsReplicator.cs
--- a/Engine/Source/Programs/UnrealCloudDDC/Jupiter/Implementation/Replication/RefsReplicator.cs
+++ b/Engine/Source/Programs/UnrealCloudDDC/Jupiter/Implementation/Replication/RefsReplicator.cs
@@ -480,7 +480,34 @@ namespace Jupiter.Implementation
             //    return currentOffset;
 
             using HttpRequestMessage referencesRequest = await BuildHttpRequest(HttpMethod.Get, new Uri($"api/v1/objects/{ns}/{objectToReplicate}/references", UriKind.Relative));
-            HttpResponseMessage referencesResponse = await _httpClient.SendAsync(referencesRequest, cancellationToken);
+            HttpResponseMessage? referencesResponse = null;
+            Exception? lastException = null;
+            const int RetryAttempts = 3;
+            for (int i = 0; i < RetryAttempts; i++)
+            {
+                try
+                {
+                    referencesResponse = await _httpClient.SendAsync(referencesRequest, cancellationToken);
+                    break;
+                }
+                catch (HttpRequestException e)
+                {
+                    referencesResponse = null;
+                    // rethrow unknown exceptions
+                    if (e.InnerException is not IOException)
+                    {
+                        throw;
+                    }
+
+                    lastException = e;
+                }
+            }
+
+            if (referencesResponse == null)
+            {
+                throw new Exception("Reference response never set", lastException);
+            }
+
             string body = await referencesResponse.Content.ReadAsStringAsync(cancellationToken);
 
             if (referencesResponse.StatusCode == HttpStatusCode.BadRequest)
