Add a comment to our future selves explaining that AWS credentials in docker compose files are not *really* live AWS credentials, but creds for the "docker composed" MinIO instance instead.

diff --git a/Engine/Source/Programs/UnrealCloudDDC/Composes/docker-compose-aws.yml b/Engine/Source/Programs/UnrealCloudDDC/Composes/docker-compose-aws.yml
--- a/Engine/Source/Programs/UnrealCloudDDC/Composes/docker-compose-aws.yml
+++ b/Engine/Source/Programs/UnrealCloudDDC/Composes/docker-compose-aws.yml
@@ -2,6 +2,8 @@ version: '3'
 services:
   unreal_cloud_ddc:
     environment:
+    # NOTE: These are not real AWS credentials, but credentials for the local MinIO instance set up below.
+    
     - AWSCredentials__AwsAccessKey=FSYPTVCPKC2IV5D1U8RE
     - AWSCredentials__AwsSecretKey=WaUf4Zq+J6YMlvbCVvTRkQJIPc7DvZwTm5UP0tPJ
 
