Unreal Cloud DDC - Increased time we wait for Scylla to be available to workaround issue with slow startup on test machines.

diff --git a/Engine/Source/Programs/UnrealCloudDDC/Jupiter/Jupiter.Startup.cs b/Engine/Source/Programs/UnrealCloudDDC/Jupiter/Jupiter.Startup.cs
--- a/Engine/Source/Programs/UnrealCloudDDC/Jupiter/Jupiter.Startup.cs
+++ b/Engine/Source/Programs/UnrealCloudDDC/Jupiter/Jupiter.Startup.cs
@@ -361,7 +361,7 @@ namespace Jupiter
             }
 
             ISession replicatedSession;
-            const int MaxRetryAttempts = 50;
+            const int MaxRetryAttempts = 100;
             int countOfAttempts = 0;
             while(true)
             {
