Unreal Cloud DDC - Stubbed out alias methods from storage client as we do not yet support it.
Should fix build issues.

diff --git a/Engine/Source/Programs/UnrealCloudDDC/Jupiter/Implementation/Bundles/StorageService.cs b/Engine/Source/Programs/UnrealCloudDDC/Jupiter/Implementation/Bundles/StorageService.cs
--- a/Engine/Source/Programs/UnrealCloudDDC/Jupiter/Implementation/Bundles/StorageService.cs
+++ b/Engine/Source/Programs/UnrealCloudDDC/Jupiter/Implementation/Bundles/StorageService.cs
@@ -81,6 +81,19 @@ public class StorageClient : IStorageClientJupiter
         return locator;
     }
 
+
+    public async Task AddAliasAsync(Utf8String name, NodeHandle handle, CancellationToken cancellationToken = default)
+    {
+        // TODO: Implement aliases
+        await Task.CompletedTask;
+    }
+
+    public async Task RemoveAliasAsync(Utf8String name, NodeHandle handle, CancellationToken cancellationToken = default)
+    {
+        // TODO: Implement aliases
+        await Task.CompletedTask;
+    }
+
     public async IAsyncEnumerable<NodeHandle> FindNodesAsync(Utf8String name, [EnumeratorCancellation] CancellationToken cancellationToken = default)
     {
         // TODO: Implement aliases
