Unreal Cloud DDC - Added more detailed error message when replicators run into unknown respones to debug issue we are currently seeing.

diff --git a/Engine/Source/Programs/UnrealCloudDDC/Jupiter/Implementation/Replication/RefsReplicator.cs b/Engine/Source/Programs/UnrealCloudDDC/Jupiter/Implementation/Replication/RefsReplicator.cs
--- a/Engine/Source/Programs/UnrealCloudDDC/Jupiter/Implementation/Replication/RefsReplicator.cs
+++ b/Engine/Source/Programs/UnrealCloudDDC/Jupiter/Implementation/Replication/RefsReplicator.cs
@@ -637,6 +637,11 @@ namespace Jupiter.Implementation
                     throw new Exception($"Unknown error when deserializing replication log events {ns} {lastBucket} {lastEvent}");
                 }
 
+                if (e.Events == null)
+                {
+                    throw new Exception($"Unknown error when deserializing replication log events {ns} {lastBucket} {lastEvent} as events were empty. Body was: {body}");
+                }
+
                 logEvents = e;
                 foreach (ReplicationLogEvent logEvent in logEvents.Events)
                 {
