Unreal Cloud DDC - Attempt at fixing json deserialization issues in the replication

diff --git a/Engine/Source/Programs/UnrealCloudDDC/Jupiter/Implementation/Replication/RefsReplicator.cs b/Engine/Source/Programs/UnrealCloudDDC/Jupiter/Implementation/Replication/RefsReplicator.cs
--- a/Engine/Source/Programs/UnrealCloudDDC/Jupiter/Implementation/Replication/RefsReplicator.cs
+++ b/Engine/Source/Programs/UnrealCloudDDC/Jupiter/Implementation/Replication/RefsReplicator.cs
@@ -44,7 +44,7 @@ namespace Jupiter.Implementation
         private bool _disposed = false;
 
         private static Histogram<long>? s_replicatedCounter;
-
+        private static JsonSerializerOptions DefaultSerializerSettings = ConfigureJsonOptions();
         public RefsReplicator(ReplicatorSettings replicatorSettings, IBlobService blobService, IHttpClientFactory httpClientFactory, IReplicationLog replicationLog, IServiceCredentials serviceCredentials, Tracer tracer, BufferedPayloadFactory bufferedPayloadFactory, ReplicationLogFactory replicationLogFactory, ILogger<RefsReplicator> logger, Meter meter)
         {
             _name = replicatorSettings.ReplicatorName;
@@ -83,6 +83,13 @@ namespace Jupiter.Implementation
             }
         }
 
+        private static JsonSerializerOptions ConfigureJsonOptions()
+        {
+            JsonSerializerOptions options = new JsonSerializerOptions();
+            BaseStartup.ConfigureJsonOptions(options);
+            return options;
+        }
+
         public void Dispose()
         {
             Dispose(true);
@@ -631,7 +638,7 @@ namespace Jupiter.Implementation
                 }
 
                 response.EnsureSuccessStatusCode();
-                ReplicationLogEvents? e = JsonSerializer.Deserialize<ReplicationLogEvents>(body);
+                ReplicationLogEvents? e = JsonSerializer.Deserialize<ReplicationLogEvents>(body, DefaultSerializerSettings);
                 if (e == null)
                 {
                     throw new Exception($"Unknown error when deserializing replication log events {ns} {lastBucket} {lastEvent}");
