Unreal Cloud DDC - More fixes as a result of the change to System.Text.Json for serialization

diff --git a/Engine/Source/Programs/UnrealCloudDDC/Jupiter/Implementation/Replication/RefsReplicator.cs b/Engine/Source/Programs/UnrealCloudDDC/Jupiter/Implementation/Replication/RefsReplicator.cs
--- a/Engine/Source/Programs/UnrealCloudDDC/Jupiter/Implementation/Replication/RefsReplicator.cs
+++ b/Engine/Source/Programs/UnrealCloudDDC/Jupiter/Implementation/Replication/RefsReplicator.cs
@@ -531,7 +531,7 @@ namespace Jupiter.Implementation
             }
             referencesResponse.EnsureSuccessStatusCode();
 
-            ResolvedReferencesResult? refs = JsonSerializer.Deserialize<ResolvedReferencesResult>(body);
+            ResolvedReferencesResult? refs = JsonSerializer.Deserialize<ResolvedReferencesResult>(body, DefaultSerializerSettings);
             if (refs == null)
             {
                 throw new Exception($"Unable to resolve references for object {objectToReplicate} in namespace {ns}");
