Unreal Cloud DDC - Retry replication of blobs if the connection is reset

diff --git a/Engine/Source/Programs/UnrealCloudDDC/Jupiter/Implementation/Replication/RefsReplicator.cs b/Engine/Source/Programs/UnrealCloudDDC/Jupiter/Implementation/Replication/RefsReplicator.cs
--- a/Engine/Source/Programs/UnrealCloudDDC/Jupiter/Implementation/Replication/RefsReplicator.cs
+++ b/Engine/Source/Programs/UnrealCloudDDC/Jupiter/Implementation/Replication/RefsReplicator.cs
@@ -16,7 +16,6 @@ using System.Threading.Tasks;
 using EpicGames.Horde.Storage;
 using Jupiter.Controllers;
 using Jupiter.Implementation.TransactionLog;
-using Jupiter.Common;
 using Jupiter.Common.Implementation;
 using Microsoft.AspNetCore.Mvc;
 using OpenTelemetry.Trace;
@@ -548,11 +547,36 @@ namespace Jupiter.Implementation
                 BlobIdentifier blobToReplicate = missingBlobs[i];
                 blobReplicationTasks[i] = Task.Run(async () =>
                 {
-                    HttpRequestMessage blobRequest = await BuildHttpRequest(HttpMethod.Get, new Uri($"api/v1/blobs/{ns}/{blobToReplicate}", UriKind.Relative));
-                    HttpResponseMessage blobResponse = await _httpClient.SendAsync(blobRequest, HttpCompletionOption.ResponseHeadersRead, cancellationToken);
-
                     _logger.LogInformation("Attempting to replicate blob {Blob} in {Namespace}.", blobToReplicate, ns);
 
+                    HttpResponseMessage? blobResponse = null;
+                    for (int i = 0; i < RetryAttempts; i++)
+                    {
+                        using HttpRequestMessage blobRequest = await BuildHttpRequest(HttpMethod.Get, new Uri($"api/v1/blobs/{ns}/{blobToReplicate}", UriKind.Relative));
+                        try
+                        {
+                            blobResponse = await _httpClient.SendAsync(blobRequest, HttpCompletionOption.ResponseHeadersRead, cancellationToken);
+                            break;
+                        }
+                        catch (HttpRequestException e)
+                        {
+                            blobResponse = null;
+                            // rethrow unknown exceptions
+                            if (e.InnerException is not IOException)
+                            {
+                                throw;
+                            }
+
+                            lastException = e;
+                        }
+                    }
+                    
+                    
+                    if (blobResponse == null)
+                    {
+                        throw new Exception("Blob response never set", lastException);
+                    }
+
                     if (blobResponse.StatusCode == HttpStatusCode.NotFound)
                     {
                         _logger.LogWarning("Failed to replicate {Blob} in {Namespace} due to it not existing.", blobToReplicate, ns);
