Unreal Cloud DDC - Fixed issue were we didn't correctly handle a ref that was missing its blob reference (such a ref is invalid as its not pointing to anything)

diff --git a/Engine/Source/Programs/UnrealCloudDDC/Jupiter/Implementation/Objects/ScyllaReferencesStore.cs b/Engine/Source/Programs/UnrealCloudDDC/Jupiter/Implementation/Objects/ScyllaReferencesStore.cs
--- a/Engine/Source/Programs/UnrealCloudDDC/Jupiter/Implementation/Objects/ScyllaReferencesStore.cs
+++ b/Engine/Source/Programs/UnrealCloudDDC/Jupiter/Implementation/Objects/ScyllaReferencesStore.cs
@@ -113,7 +113,7 @@ namespace Jupiter.Implementation
 
             try
             {
-                o.ThrowIfRequiredFieldIsMissing(includePayload);
+                o.ThrowIfRequiredFieldIsMissing();
             }
             catch (Exception e)
             {
@@ -491,7 +491,7 @@ namespace Jupiter.Implementation
         [Cassandra.Mapping.Attributes.Column("last_access_time")]
         public DateTime LastAccessTime { get; set; }
 
-        public void ThrowIfRequiredFieldIsMissing(bool includePayload)
+        public void ThrowIfRequiredFieldIsMissing()
         {
             if (string.IsNullOrEmpty(Namespace))
             {
@@ -508,7 +508,7 @@ namespace Jupiter.Implementation
                 throw new InvalidOperationException("Name was not valid");
             }
 
-            if (PayloadHash == null && includePayload)
+            if (PayloadHash == null)
             {
                 throw new ArgumentException("PayloadHash was not valid");
             }
