Unreal Cloud DDC - Updated migration options to assume migration has happened (so disabled fallback paths)

diff --git a/Engine/Source/Programs/UnrealCloudDDC/Jupiter/Implementation/Objects/ScyllaSettings.cs b/Engine/Source/Programs/UnrealCloudDDC/Jupiter/Implementation/Objects/ScyllaSettings.cs
--- a/Engine/Source/Programs/UnrealCloudDDC/Jupiter/Implementation/Objects/ScyllaSettings.cs
+++ b/Engine/Source/Programs/UnrealCloudDDC/Jupiter/Implementation/Objects/ScyllaSettings.cs
@@ -60,14 +60,13 @@ public class ScyllaSettings : IValidatableObject
 
     /// <summary>
     /// Toggle to enable reading objects from the object_last_access_time_v2 table
-    /// Experimental option.
     /// </summary>
-    public bool ListObjectsFromLastAccessTable { get; set; } = false;
+    public bool ListObjectsFromLastAccessTable { get; set; } = true;
 
     /// <summary>
     /// Set to update the last access field in the object table, this is legacy and causes extra compaction work on the database
     /// </summary>
-    public bool UpdateLegacyLastAccessField { get; set; } = true;
+    public bool UpdateLegacyLastAccessField { get; set; } = false;
 
     /// <summary>
     /// Read timeout in milliseconds
@@ -95,12 +94,12 @@ public class ScyllaSettings : IValidatableObject
     /// <summary>
     /// Enable to migrate data from blob_index table to blob_index_v2 table
     /// </summary>
-    public bool MigrateFromOldBlobIndex { get; set; } = true;
+    public bool MigrateFromOldBlobIndex { get; set; } = false;
 
     /// <summary>
     /// Enable to use list namespaces from the old table, this can set for a period while migration from old to new table happens
     /// </summary>
-    public bool ListObjectsFromOldNamespaceTable { get; set; } = true;
+    public bool ListObjectsFromOldNamespaceTable { get; set; } = false;
 
     public IEnumerable<ValidationResult> Validate(ValidationContext validationContext)
     {
diff --git a/Engine/Source/Programs/UnrealCloudDDC/changelog.md b/Engine/Source/Programs/UnrealCloudDDC/changelog.md
--- a/Engine/Source/Programs/UnrealCloudDDC/changelog.md
+++ b/Engine/Source/Programs/UnrealCloudDDC/changelog.md
@@ -1,5 +1,7 @@
 # Unreleased
 
+* Migration options from `0.3.0` have been updated to assume you have migrated by default.
+
 # 0.3.0
 * Azure blob storage now supports storage pools
 * Last access table refactoring - Moved the last accessing tracking out of the objects table and into a seperate table. Saves on compation work for Scylla. Set `Scylla.ListObjectsFromLastAccessTable` to migrate GC refs to use this new table (will be default in the next release).
