Unreal Cloud DDC - Added retry when reading the refs, similarly to what we already do when replicating the blobs.

diff --git a/Engine/Source/Programs/UnrealCloudDDC/Jupiter/Implementation/Replication/RefsReplicator.cs b/Engine/Source/Programs/UnrealCloudDDC/Jupiter/Implementation/Replication/RefsReplicator.cs
--- a/Engine/Source/Programs/UnrealCloudDDC/Jupiter/Implementation/Replication/RefsReplicator.cs
+++ b/Engine/Source/Programs/UnrealCloudDDC/Jupiter/Implementation/Replication/RefsReplicator.cs
@@ -635,8 +635,37 @@ namespace Jupiter.Implementation
                 }
 
                 hasRunOnce = true;
-                using HttpRequestMessage request = await BuildHttpRequest(HttpMethod.Get, new Uri(url.ToString(), UriKind.Relative));
-                HttpResponseMessage response = await _httpClient.SendAsync(request, cancellationToken);
+
+                HttpResponseMessage? response = null;
+                Exception? lastException = null;
+                const int RetryAttempts = 3;
+                for (int i = 0; i < RetryAttempts; i++)
+                {
+                    using HttpRequestMessage request = await BuildHttpRequest(HttpMethod.Get, new Uri(url.ToString(), UriKind.Relative));
+
+                    try
+                    {
+                        response = await _httpClient.SendAsync(request, cancellationToken);
+                        break;
+                    }
+                    catch (HttpRequestException e)
+                    {
+                        response = null;
+                        // rethrow unknown exceptions
+                        if (e.InnerException is not IOException)
+                        {
+                            throw;
+                        }
+
+                        lastException = e;
+                    }
+                }
+
+                if (response == null)
+                {
+                    throw new Exception("Ref response never set", lastException);
+                }
+
                 string body = await response.Content.ReadAsStringAsync(cancellationToken);
                 if (response.StatusCode == HttpStatusCode.BadRequest)
                 {
@@ -662,18 +691,18 @@ namespace Jupiter.Implementation
                 }
 
                 response.EnsureSuccessStatusCode();
-                ReplicationLogEvents? e = JsonSerializer.Deserialize<ReplicationLogEvents>(body, DefaultSerializerSettings);
-                if (e == null)
+                ReplicationLogEvents? replicationLogEvents = JsonSerializer.Deserialize<ReplicationLogEvents>(body, DefaultSerializerSettings);
+                if (replicationLogEvents == null)
                 {
                     throw new Exception($"Unknown error when deserializing replication log events {ns} {lastBucket} {lastEvent}");
                 }
 
-                if (e.Events == null)
+                if (replicationLogEvents.Events == null)
                 {
                     throw new Exception($"Unknown error when deserializing replication log events {ns} {lastBucket} {lastEvent} as events were empty. Body was: {body}");
                 }
 
-                logEvents = e;
+                logEvents = replicationLogEvents;
                 foreach (ReplicationLogEvent logEvent in logEvents.Events)
                 {
                     yield return logEvent;
