Unreal Cloud DDC - Added cmdlet to upload benchmark to our Dev account

diff --git a/Engine/Source/Programs/UnrealCloudDDC/Benchmarks/benchmark-suite/cmdlets.ps1 b/Engine/Source/Programs/UnrealCloudDDC/Benchmarks/benchmark-suite/cmdlets.ps1
--- a/Engine/Source/Programs/UnrealCloudDDC/Benchmarks/benchmark-suite/cmdlets.ps1
+++ b/Engine/Source/Programs/UnrealCloudDDC/Benchmarks/benchmark-suite/cmdlets.ps1
@@ -10,6 +10,18 @@ function Upload-BenchmarkECR {
     docker push 728559092788.dkr.ecr.us-east-1.amazonaws.com/jupiter_benchmark:latest
 }
 
+function Upload-BenchmarkECR-Dev {
+    Write-Host "Logging in to AWS ECR (Dev)"
+    aws ecr get-login-password --region us-east-1 | docker login --username AWS --password-stdin 730468387612.dkr.ecr.us-east-1.amazonaws.com
+
+    Write-Host "Starting docker build"
+    docker build -t benchmarker .
+
+    Write-Host "Uploading to ECR (Dev)"
+    docker tag benchmarker 730468387612.dkr.ecr.us-east-1.amazonaws.com/jupiter_benchmark:latest
+    docker push 730468387612.dkr.ecr.us-east-1.amazonaws.com/jupiter_benchmark:latest
+}
+
 function Upload-BenchmarkS3 {
     $scriptDir = $PSScriptRoot
     $archivePath = (Join-Path $scriptDir "benchmark.zip")
