Unreal Cloud DDC - Added options to delay termination of pods which is useful when working with a NLB with ip targets as that is fairly slow to register new targets.

diff --git a/Engine/Source/Programs/UnrealCloudDDC/Helm/UnrealCloudDDC/templates/deployment.yaml b/Engine/Source/Programs/UnrealCloudDDC/Helm/UnrealCloudDDC/templates/deployment.yaml
--- a/Engine/Source/Programs/UnrealCloudDDC/Helm/UnrealCloudDDC/templates/deployment.yaml
+++ b/Engine/Source/Programs/UnrealCloudDDC/Helm/UnrealCloudDDC/templates/deployment.yaml
@@ -27,6 +27,9 @@ spec:
       imagePullSecrets:
         {{- toYaml . | nindent 8 }}
       {{- end }}
+      {{- if .Values.terminationGracePeriodSeconds }}
+      terminationGracePeriodSeconds: {{ .Values.terminationGracePeriodSeconds }}
+      {{- end }}
       {{- if .Values.affinity }}
       affinity: {{- include "common.tplvalues.render" (dict "value" .Values.affinity "context" $) | nindent 8 }}
       {{- else }}
@@ -91,6 +94,15 @@ spec:
           successThreshold: {{ .Values.readinessProbe.successThreshold }}
           failureThreshold: {{ .Values.readinessProbe.failureThreshold }}
         {{- end }}
+        {{- if .Values.preStopHookDelay }}
+        lifecycle:
+          preStop:
+            exec:
+              command: 
+              # Load balancers can be slow to deregister instances, this can be used to introduce a delay to provide zero downtime
+              - "sleep"
+              - "{{ .Values.preStopHookDelay }}"
+        {{- end }}
         {{- if .Values.resources }}
         resources: {{- toYaml .Values.resources | nindent 12 }}
         {{- end }}
@@ -159,6 +171,15 @@ spec:
           successThreshold: {{ .Values.nginx.readinessProbe.successThreshold }}
           failureThreshold: {{ .Values.nginx.readinessProbe.failureThreshold }}
         {{- end }}
+        {{- if .Values.preStopHookDelay }}
+        lifecycle:
+          preStop:
+            exec:
+              command: 
+              # Load balancers can be slow to deregister instances, this can be used to introduce a delay to provide zero downtime
+              - "sleep"
+              - "{{ .Values.preStopHookDelay }}"
+        {{- end }}
         volumeMounts:
           - name: nginx
             mountPath: /nginx
