Unreal Cloud DDC - Fixed failing tests - we can not selectivley create table layouts as we always create the types that bind to them causing Scylla OEM to fail if the tables are missing.

diff --git a/Engine/Source/Programs/UnrealCloudDDC/Jupiter/Implementation/Objects/ScyllaReferencesStore.cs b/Engine/Source/Programs/UnrealCloudDDC/Jupiter/Implementation/Objects/ScyllaReferencesStore.cs
--- a/Engine/Source/Programs/UnrealCloudDDC/Jupiter/Implementation/Objects/ScyllaReferencesStore.cs
+++ b/Engine/Source/Programs/UnrealCloudDDC/Jupiter/Implementation/Objects/ScyllaReferencesStore.cs
@@ -69,15 +69,12 @@ namespace Jupiter.Implementation
             );"
             ));
 
-            if (_settings.CurrentValue.ListObjectsFromOldNamespaceTable)
-            {
-                _session.Execute(new SimpleStatement(@"CREATE TABLE IF NOT EXISTS buckets (
-                    namespace text, 
-                    bucket set<text>, 
-                    PRIMARY KEY (namespace)
-                );"
-                ));
-            }
+            _session.Execute(new SimpleStatement(@"CREATE TABLE IF NOT EXISTS buckets (
+                namespace text, 
+                bucket set<text>, 
+                PRIMARY KEY (namespace)
+            );"
+            ));
 
             // BYPASS CACHE is a scylla specific extension to disable populating the cache, should be ignored by other cassandra dbs
             string cqlOptions = scyllaSessionManager.IsScylla ? "BYPASS CACHE" : "";
