Unreal Cloud DDC - Added support for https redirect via nginx if using a level 4 load balancer.

diff --git a/Engine/Source/Programs/UnrealCloudDDC/Helm/UnrealCloudDDC/templates/configmap-nginx.yaml b/Engine/Source/Programs/UnrealCloudDDC/Helm/UnrealCloudDDC/templates/configmap-nginx.yaml
--- a/Engine/Source/Programs/UnrealCloudDDC/Helm/UnrealCloudDDC/templates/configmap-nginx.yaml
+++ b/Engine/Source/Programs/UnrealCloudDDC/Helm/UnrealCloudDDC/templates/configmap-nginx.yaml
@@ -168,6 +168,7 @@ data:
         }
       }
 
+      # nginx status
       server {
         listen 18080;
 
@@ -182,5 +183,11 @@ data:
           return 404;
         }
       }
+
+      # https redirect
+      server {
+          listen 18090;
+          return 301 https://$server_name$request_uri;
+      }
     }  
 {{ end }}
diff --git a/Engine/Source/Programs/UnrealCloudDDC/Helm/UnrealCloudDDC/templates/deployment.yaml b/Engine/Source/Programs/UnrealCloudDDC/Helm/UnrealCloudDDC/templates/deployment.yaml
--- a/Engine/Source/Programs/UnrealCloudDDC/Helm/UnrealCloudDDC/templates/deployment.yaml
+++ b/Engine/Source/Programs/UnrealCloudDDC/Helm/UnrealCloudDDC/templates/deployment.yaml
@@ -149,6 +149,8 @@ spec:
         - name: nginx-corp-h2
           containerPort: 8092
         - containerPort: 18080
+        - name: nginx-https-redirect
+          containerPort: 18090
         {{- if .Values.nginx.livenessProbe.enabled }}
         livenessProbe:
           httpGet:
