Unreal Cloud DDC - Use a shorter identifier for https redirect to be in line with kubernetes requirments

diff --git a/Engine/Source/Programs/UnrealCloudDDC/Helm/UnrealCloudDDC/templates/deployment.yaml b/Engine/Source/Programs/UnrealCloudDDC/Helm/UnrealCloudDDC/templates/deployment.yaml
--- a/Engine/Source/Programs/UnrealCloudDDC/Helm/UnrealCloudDDC/templates/deployment.yaml
+++ b/Engine/Source/Programs/UnrealCloudDDC/Helm/UnrealCloudDDC/templates/deployment.yaml
@@ -149,7 +149,7 @@ spec:
         - name: nginx-corp-h2
           containerPort: 8092
         - containerPort: 18080
-        - name: nginx-https-redirect
+        - name: nginx-redirect
           containerPort: 18090
         {{- if .Values.nginx.livenessProbe.enabled }}
         livenessProbe:
