Unreal Cloud DDC - Update nginx rewrite to use host variable

diff --git a/Engine/Source/Programs/UnrealCloudDDC/Helm/UnrealCloudDDC/templates/configmap-nginx.yaml b/Engine/Source/Programs/UnrealCloudDDC/Helm/UnrealCloudDDC/templates/configmap-nginx.yaml
--- a/Engine/Source/Programs/UnrealCloudDDC/Helm/UnrealCloudDDC/templates/configmap-nginx.yaml
+++ b/Engine/Source/Programs/UnrealCloudDDC/Helm/UnrealCloudDDC/templates/configmap-nginx.yaml
@@ -187,7 +187,7 @@ data:
       # https redirect
       server {
           listen 18090;
-          return 301 https://$server_name$request_uri;
+          return 301 https://$host$request_uri;
       }
     }  
 {{ end }}
