Unreal Cloud DDC - Shortcut the serving of nginx routes that doesn't exist to avoid errors from nginx (still returning 404s for it)

diff --git a/Engine/Source/Programs/UnrealCloudDDC/Helm/UnrealCloudDDC/templates/configmap-nginx.yaml b/Engine/Source/Programs/UnrealCloudDDC/Helm/UnrealCloudDDC/templates/configmap-nginx.yaml
--- a/Engine/Source/Programs/UnrealCloudDDC/Helm/UnrealCloudDDC/templates/configmap-nginx.yaml
+++ b/Engine/Source/Programs/UnrealCloudDDC/Helm/UnrealCloudDDC/templates/configmap-nginx.yaml
@@ -64,6 +64,8 @@ data:
           default_type "";
           return 200; # error_code is ignored
         }
+
+        return 404;
       }
 
       # http 1.1 corp port
@@ -98,6 +100,8 @@ data:
           default_type "";
           return 200; # error_code is ignored
         }
+
+        return 404;
       }
 
       # public http2 port
@@ -132,6 +136,8 @@ data:
           default_type "";
           return 200; # error_code is ignored
         }
+
+        return 404;
       }
 
       # corp port http2
@@ -166,6 +172,8 @@ data:
           default_type "";
           return 200; # error_code is ignored
         }
+
+        return 404;
       }
 
       # nginx status
