Unreal Cloud DDC - Added a better catch all configuration for unknown routes to avoid nginx write errors.

diff --git a/Engine/Source/Programs/UnrealCloudDDC/Helm/UnrealCloudDDC/templates/configmap-nginx.yaml b/Engine/Source/Programs/UnrealCloudDDC/Helm/UnrealCloudDDC/templates/configmap-nginx.yaml
--- a/Engine/Source/Programs/UnrealCloudDDC/Helm/UnrealCloudDDC/templates/configmap-nginx.yaml
+++ b/Engine/Source/Programs/UnrealCloudDDC/Helm/UnrealCloudDDC/templates/configmap-nginx.yaml
@@ -59,13 +59,15 @@ data:
           {{ if .Values.nginx.addXAccelDebugHeader }} add_header X-Jupiter-XAccel-Used true; {{ end }}
         }
 
+        location / {
+            return 404;
+        }
+
         error_page 404 @error_page;
         location @error_page {
           default_type "";
           return 200; # error_code is ignored
         }
-
-        return 404;
       }
 
       # http 1.1 corp port
@@ -95,13 +97,15 @@ data:
           {{ if .Values.nginx.addXAccelDebugHeader }} add_header X-Jupiter-XAccel-Used true; {{ end }}
         }
 
+        location / {
+            return 404;
+        }
+
         error_page 404 @error_page;
         location @error_page {
           default_type "";
           return 200; # error_code is ignored
         }
-
-        return 404;
       }
 
       # public http2 port
@@ -131,13 +135,15 @@ data:
           {{ if .Values.nginx.addXAccelDebugHeader }} add_header X-Jupiter-XAccel-Used true; {{ end }}
         }
 
+        location / {
+            return 404;
+        }
+
         error_page 404 @error_page;
         location @error_page {
           default_type "";
           return 200; # error_code is ignored
         }
-
-        return 404;
       }
 
       # corp port http2
@@ -167,13 +173,15 @@ data:
             {{ if .Values.nginx.addXAccelDebugHeader }}add_header X-Jupiter-XAccel-Used true; {{ end }}
         }
 
+        location / {
+            return 404;
+        }
+
         error_page 404 @error_page;
         location @error_page {
           default_type "";
           return 200; # error_code is ignored
         }
-
-        return 404;
       }
 
       # nginx status
