Unreal Cloud DDC - Avoid caching http clients that are allocated via a http client factory as these are recommended to be short lived when allocated via a factory (unlike normal http clients)

diff --git a/Engine/Source/Programs/UnrealCloudDDC/Jupiter/Implementation/Blob/PeerBlobStore.cs b/Engine/Source/Programs/UnrealCloudDDC/Jupiter/Implementation/Blob/PeerBlobStore.cs
--- a/Engine/Source/Programs/UnrealCloudDDC/Jupiter/Implementation/Blob/PeerBlobStore.cs
+++ b/Engine/Source/Programs/UnrealCloudDDC/Jupiter/Implementation/Blob/PeerBlobStore.cs
@@ -25,7 +25,6 @@ namespace Jupiter.Implementation
         private readonly IHttpClientFactory _httpClientFactory;
         private readonly IServiceCredentials _serviceCredentials;
         private readonly ILogger _logger;
-        private readonly ConcurrentDictionary<string, HttpClient> _httpClients = new ConcurrentDictionary<string, HttpClient>();
 
         public PeerBlobStore(IPeerServiceDiscovery serviceDiscovery, IHttpClientFactory httpClientFactory, IServiceCredentials serviceCredentials, ILogger<PeerBlobStore> logger)
         {
@@ -159,15 +158,12 @@ namespace Jupiter.Implementation
 
         private HttpClient GetHttpClient(string instance)
         {
-            return _httpClients.GetOrAdd(instance, s =>
-            {
-                HttpClient httpClient = _httpClientFactory.CreateClient(instance);
-                httpClient.BaseAddress = new Uri($"http://{instance}");
+            HttpClient httpClient = _httpClientFactory.CreateClient(instance);
+            httpClient.BaseAddress = new Uri($"http://{instance}");
 
-                // for these connections to be useful they need to be fast - so timeout quickly if we can not establish the connection
-                httpClient.Timeout = TimeSpan.FromSeconds(1.0);
-                return httpClient;
-            });
+            // for these connections to be useful they need to be fast - so timeout quickly if we can not establish the connection
+            httpClient.Timeout = TimeSpan.FromSeconds(1.0);
+            return httpClient;
         }
 
         public Task<Uri?> GetObjectByRedirect(NamespaceId ns, BlobIdentifier identifier)
