Jupiter build fixes.

diff --git a/Engine/Source/Programs/UnrealCloudDDC/Jupiter.Common/Implementation/BlobIdentifier.cs b/Engine/Source/Programs/UnrealCloudDDC/Jupiter.Common/Implementation/BlobIdentifier.cs
--- a/Engine/Source/Programs/UnrealCloudDDC/Jupiter.Common/Implementation/BlobIdentifier.cs
+++ b/Engine/Source/Programs/UnrealCloudDDC/Jupiter.Common/Implementation/BlobIdentifier.cs
@@ -156,12 +156,12 @@ namespace Jupiter.Implementation
 
         public static BlobIdentifier FromBlobLocator(BlobLocator locator)
         {
-            return new BlobIdentifier(Encoding.UTF8.GetBytes(locator.BlobId.ToString()));
+            return new BlobIdentifier(Encoding.UTF8.GetBytes(locator.ToString()));
         }
 
         public BlobLocator AsBlobLocator()
         {
-            return new BlobLocator(HostId.Empty, new BlobId(Encoding.UTF8.GetString(HashData)));
+            return new BlobLocator(Encoding.UTF8.GetString(HashData));
         }
     }
 
diff --git a/Engine/Source/Programs/UnrealCloudDDC/Jupiter/Controllers/BundlesController.cs b/Engine/Source/Programs/UnrealCloudDDC/Jupiter/Controllers/BundlesController.cs
--- a/Engine/Source/Programs/UnrealCloudDDC/Jupiter/Controllers/BundlesController.cs
+++ b/Engine/Source/Programs/UnrealCloudDDC/Jupiter/Controllers/BundlesController.cs
@@ -328,7 +328,7 @@ namespace Jupiter.Controllers
                 return result;
             }
             IStorageClientJupiter client = await _storageService.GetClientAsync(namespaceId, cancellationToken);
-            HashedNodeLocator target = new HashedNodeLocator(request.Hash, request.Blob, request.ExportIdx);
+            NodeLocator target = new NodeLocator(request.Hash, request.Blob, request.ExportIdx);
             await client.WriteRefTargetAsync(refName, target, request.Options, cancellationToken);
 
             return Ok();
@@ -466,7 +466,7 @@ namespace Jupiter.Controllers
 
             object content;
 
-            BlobData nodeData = await reader.ReadNodeDataAsync(new NodeLocator(locator, exportIdx), cancellationToken);
+            BlobData nodeData = await reader.ReadNodeDataAsync(new NodeLocator(IoHash.Zero, locator, exportIdx), cancellationToken);
 
             Node node = Node.Deserialize(nodeData);
             switch (node)
diff --git a/Engine/Source/Programs/UnrealCloudDDC/Jupiter/Implementation/Bundles/StorageService.cs b/Engine/Source/Programs/UnrealCloudDDC/Jupiter/Implementation/Bundles/StorageService.cs
--- a/Engine/Source/Programs/UnrealCloudDDC/Jupiter/Implementation/Bundles/StorageService.cs
+++ b/Engine/Source/Programs/UnrealCloudDDC/Jupiter/Implementation/Bundles/StorageService.cs
@@ -21,7 +21,7 @@ public interface IStorageClientJupiter : IStorageClient
     Task<(BlobLocator Locator, Uri UploadUrl)?> GetWriteRedirectAsync(string prefix, CancellationToken cancellationToken);
     bool SupportsRedirects { get; set; }
     Task<Uri?> GetReadRedirectAsync(BlobLocator locator, CancellationToken cancellationToken);
-    Task WriteRefTargetAsync(RefName refName, HashedNodeLocator target, RefOptions? requestOptions, CancellationToken cancellationToken);
+    Task WriteRefTargetAsync(RefName refName, NodeLocator target, RefOptions? requestOptions, CancellationToken cancellationToken);
 }
 
 public interface IStorageService
@@ -176,7 +176,7 @@ public class StorageClient : IStorageClientJupiter
             BlobLocator blobLocator = new BlobLocator(inlinePayload.BlobLocator);
             int exportId = inlinePayload.ExportId;
 
-            return new FlushedNodeHandle(_treeReader, nodeHash, new NodeLocator(blobLocator, exportId));
+            return new FlushedNodeHandle(_treeReader, new NodeLocator(nodeHash, blobLocator, exportId));
         }
         catch (ObjectNotFoundException )
         {
@@ -187,7 +187,7 @@ public class StorageClient : IStorageClientJupiter
     public async Task<BlobHandle> WriteRefAsync(RefName name, Bundle bundle, int exportIdx, Utf8String prefix = default, RefOptions? options = null, CancellationToken cancellationToken = default)
     {
         BlobLocator locator = await this.WriteBundleAsync(bundle, prefix, cancellationToken);
-        BlobHandle target = new FlushedNodeHandle(_treeReader, bundle.Header.Exports[exportIdx].Hash, new NodeLocator(locator, exportIdx));
+        BlobHandle target = new FlushedNodeHandle(_treeReader, new NodeLocator(bundle.Header.Exports[exportIdx].Hash, locator, exportIdx));
         await WriteRefTargetAsync(name, target, options, cancellationToken);
 
         return target;
@@ -195,10 +195,10 @@ public class StorageClient : IStorageClientJupiter
 
     public Task WriteRefTargetAsync(RefName refName, BlobHandle target, RefOptions? requestOptions, CancellationToken cancellationToken)
     {
-        return WriteRefTargetAsync(refName, new HashedNodeLocator(target.Hash, target.GetLocator()), requestOptions, cancellationToken);
+        return WriteRefTargetAsync(refName, target.GetLocator(), requestOptions, cancellationToken);
     }
 
-    public async Task WriteRefTargetAsync(RefName refName, HashedNodeLocator target, RefOptions? requestOptions, CancellationToken cancellationToken)
+    public async Task WriteRefTargetAsync(RefName refName, NodeLocator target, RefOptions? requestOptions, CancellationToken cancellationToken)
     {
         BlobIdentifier bundleBlob = BlobIdentifier.FromBlobLocator(target.Blob);
         IoHashKey refKey = IoHashKey.FromName(refName.ToString());
