Unreal Cloud DDC - Added prepareNvmeFilesystem feature in our helm chart which standardizes the preperation of attached nvme drives, this also supports creating raid0 if multple drives are present.

diff --git a/Engine/Source/Programs/UnrealCloudDDC/Helm/UnrealCloudDDC/templates/deployment.yaml b/Engine/Source/Programs/UnrealCloudDDC/Helm/UnrealCloudDDC/templates/deployment.yaml
--- a/Engine/Source/Programs/UnrealCloudDDC/Helm/UnrealCloudDDC/templates/deployment.yaml
+++ b/Engine/Source/Programs/UnrealCloudDDC/Helm/UnrealCloudDDC/templates/deployment.yaml
@@ -44,8 +44,34 @@ spec:
       {{- if .Values.tolerations }}
       tolerations: {{- include "common.tplvalues.render" (dict "value" .Values.tolerations "context" $) | nindent 8 }}
       {{- end }}
+      {{- if or .Values.initContainers .Values.prepareNvmeFilesystem.enabled }}
+      initContainers: 
       {{- if .Values.initContainers }}
-      initContainers: {{- include "common.tplvalues.render" (dict "value" .Values.initContainers "context" $) | nindent 8 }}
+      {{- include "common.tplvalues.render" (dict "value" .Values.initContainers "context" $) | nindent 6 }}
+      {{- end }}
+      {{- if .Values.prepareNvmeFilesystem.enabled }}
+      {{- $drivesCount := len .Values.prepareNvmeFilesystem.drives }}
+      - name: nvme-mount
+        image: ubuntu
+        securityContext:
+          privileged: true
+        command: ["/bin/sh"]
+      {{- if eq $drivesCount 0 }}
+      {{ fail "You have to specify at least one drivepath to prepare when enabling the prepareNvmeFilesystem" }}
+      {{- else if gt $drivesCount 1 }}
+        args:
+        - -c
+        - "apt-get update && apt-get -y install mdadm && mdadm --create --verbose /dev/md0 --level=0 --raid-devices={{ $drivesCount }} {{ join " " .Values.prepareNvmeFilesystem.drives }} && mkfs.ext4 -E nodiscard /dev/md0 && mount /dev/md0 /data-nvme || exit 0"
+      {{- else }}
+        args: 
+        - -c
+        - mkfs.ext4 -E nodiscard {{ first .Values.prepareNvmeFilesystem.drives }} && mount {{ first .Values.prepareNvmeFilesystem.drives }} /data-nvme || exit 0
+      {{- end }}
+      {{- end }}   
+        volumeMounts:
+        - mountPath: /data-nvme
+          name: data
+          mountPropagation: Bidirectional
       {{- end }}
       containers:
       - name: {{ include "common.names.name" . }}
diff --git a/Engine/Source/Programs/UnrealCloudDDC/Helm/UnrealCloudDDC/values.yaml b/Engine/Source/Programs/UnrealCloudDDC/Helm/UnrealCloudDDC/values.yaml
--- a/Engine/Source/Programs/UnrealCloudDDC/Helm/UnrealCloudDDC/values.yaml
+++ b/Engine/Source/Programs/UnrealCloudDDC/Helm/UnrealCloudDDC/values.yaml
@@ -409,4 +409,14 @@ worker:
     ##   - e2e-az1
     ##   - e2e-az2
     ##
-    values: []
\ No newline at end of file
+    values: []
+
+
+## Prepare Nvme filesystem configuration
+## Optional feature that can be enabled to run a initContainer that prepares a already attached nvme drive
+## This works well with VMs with locally attached nvme drives
+## If you enable it you should set the list of drives that you expect to be enabled
+## if more then 1 drive is specified we raid0 these together into a single drive
+prepareNvmeFilesystem:
+  enabled: false
+  drives: []
\ No newline at end of file
diff --git a/Engine/Source/Programs/UnrealCloudDDC/changelog.md b/Engine/Source/Programs/UnrealCloudDDC/changelog.md
--- a/Engine/Source/Programs/UnrealCloudDDC/changelog.md
+++ b/Engine/Source/Programs/UnrealCloudDDC/changelog.md
@@ -1,7 +1,7 @@
 # Unreleased
 
 * Migration options from `0.3.0` have been updated to assume you have migrated by default.
-
+* Added `prepareNvmeFilesystem` section in Helm chart that creates a initContainer which will format a attached nvme drive.
 # 0.3.0
 * Azure blob storage now supports storage pools
 * Last access table refactoring - Moved the last accessing tracking out of the objects table and into a seperate table. Saves on compation work for Scylla. Set `Scylla.ListObjectsFromLastAccessTable` to migrate GC refs to use this new table (will be default in the next release).
