Unreal Cloud DDC - Removing usage of volumes in compose for tests. We do not want state to carry over and this should help speed up the start of the dbs and esspecially avoid issues with Scylla doing alot of compaction on startup causing tests to timeout occasionally.

diff --git a/Engine/Source/Programs/UnrealCloudDDC/Composes/docker-compose-tests.yml b/Engine/Source/Programs/UnrealCloudDDC/Composes/docker-compose-tests.yml
--- a/Engine/Source/Programs/UnrealCloudDDC/Composes/docker-compose-tests.yml
+++ b/Engine/Source/Programs/UnrealCloudDDC/Composes/docker-compose-tests.yml
@@ -9,8 +9,6 @@ services:
     environment:
     - "MINIO_ACCESS_KEY=FSYPTVCPKC2IV5D1U8RE"
     - "MINIO_SECRET_KEY=WaUf4Zq+J6YMlvbCVvTRkQJIPc7DvZwTm5UP0tPJ"
-    volumes:
-    - s3-volume:/data
     networks:
     - network_bridge
     ports:
@@ -36,8 +34,6 @@ services:
     environment:
     - "LOCAL_STORAGE_ACCOUNT_NAME=blobaccount"
     - "LOCAL_STORAGE_ACCOUNT_KEY=iU6uTvlF1ysppmft+NO5lAD0E3hwrAORr5Rb5xcBWUgEz/OicrSkFxwZYMNK5XL29/wXZKGOoOVSW040nAOfPg=="
-    volumes:
-    - blob-volume:/data
     networks:
     - network_bridge
     ports:
@@ -48,16 +44,10 @@ services:
     # Use to disable logging which can be spammy but also useful
     #logging:
     # driver: none
-    volumes:
-    - scylla-volume:/var/lib/scylla
     networks:
     - network_bridge
     ports:
     - "9042:9042" 
-volumes:
-  s3-volume:
-  blob-volume:
-  scylla-volume:
 
 networks:
   network_bridge:
