Jupiter: Separate interfaces and implementations into separate files, so we can just merge the interface into Horde. Also move NamespacePolicy into NamespacePolicyResolver.cs.

diff --git a/Engine/Source/Programs/UnrealCloudDDC/Jupiter.Common/BaseStartup.cs b/Engine/Source/Programs/UnrealCloudDDC/Jupiter.Common/BaseStartup.cs
--- a/Engine/Source/Programs/UnrealCloudDDC/Jupiter.Common/BaseStartup.cs
+++ b/Engine/Source/Programs/UnrealCloudDDC/Jupiter.Common/BaseStartup.cs
@@ -729,51 +729,4 @@ namespace Jupiter
         /// </summary>
         public bool AssumeLocalConnectionsHasFullAccess { get; set; }
     }
-
-    public class NamespaceSettings
-    {
-        [System.Diagnostics.CodeAnalysis.SuppressMessage("Usage", "CA2227:Collection properties should be read only", Justification = "Used by the configuration system")]
-        public Dictionary<string, NamespacePolicy> Policies { get; set; } = new Dictionary<string, NamespacePolicy>();
-    }
-
-    public class NamespacePolicy
-    {
-        [System.Diagnostics.CodeAnalysis.SuppressMessage("Usage", "CA2227:Collection properties should be read only", Justification = "Used by the configuration system")]
-        public List<AclEntry> Acls { get; set; } = new List<AclEntry>();
-        public string StoragePool { get; set; } = "";
-
-        public bool LastAccessTracking { get; set; } = true;
-        public bool OnDemandReplication { get; set; } = false;
-        public bool UseBlobIndexForExists { get; set; } = false;
-        public bool UseBlobIndexForSlowExists { get; set; } = false;
-        public bool IsPublicNamespace { get; set; } = true;
-        public NamespaceId? FallbackNamespace { get; set; } = null;
-        public bool PopulateFallbackNamespaceOnUpload { get; set; } = true;
-
-        public bool UseContentAddressedStorage { get; set; } = true;
-
-        public enum StoragePoolGCMethod  {
-            /// <summary>
-            /// Never run GC on this namespace
-            /// </summary>
-            None, 
-            /// <summary>
-            /// Apply last access deletion, objects not used for a duration set in GCSettings will be removed
-            /// </summary>
-            LastAccess, 
-            /// <summary>
-            /// Objects are removed after DefaultTTL time has passed, no matter if they are used or not
-            /// </summary>
-            TTL,
-            /// <summary>
-            /// Always GC references to this namespace, used to opt in to cleaning out old data
-            /// </summary>
-            Always
-        };
-
-        public StoragePoolGCMethod GcMethod { get; set; } = StoragePoolGCMethod.LastAccess;
-
-        public TimeSpan DefaultTTL { get; set; } = TimeSpan.FromDays(14);
-        public bool AllowRedirectUris { get; set; } = false;
-    }
 }
diff --git a/Engine/Source/Programs/UnrealCloudDDC/Jupiter.Common/Implementation/NamespacePolicyResolver.cs b/Engine/Source/Programs/UnrealCloudDDC/Jupiter.Common/Implementation/NamespacePolicyResolver.cs
--- a/Engine/Source/Programs/UnrealCloudDDC/Jupiter.Common/Implementation/NamespacePolicyResolver.cs
+++ b/Engine/Source/Programs/UnrealCloudDDC/Jupiter.Common/Implementation/NamespacePolicyResolver.cs
@@ -15,6 +15,54 @@ namespace Jupiter.Common
         static NamespaceId JupiterInternalNamespace => new NamespaceId("jupiter-internal");
     }
 
+    public class NamespaceSettings
+    {
+        [System.Diagnostics.CodeAnalysis.SuppressMessage("Usage", "CA2227:Collection properties should be read only", Justification = "Used by the configuration system")]
+        public Dictionary<string, NamespacePolicy> Policies { get; set; } = new Dictionary<string, NamespacePolicy>();
+    }
+
+    public class NamespacePolicy
+    {
+        [System.Diagnostics.CodeAnalysis.SuppressMessage("Usage", "CA2227:Collection properties should be read only", Justification = "Used by the configuration system")]
+        public List<AclEntry> Acls { get; set; } = new List<AclEntry>();
+        public string StoragePool { get; set; } = "";
+
+        public bool LastAccessTracking { get; set; } = true;
+        public bool OnDemandReplication { get; set; } = false;
+        public bool UseBlobIndexForExists { get; set; } = false;
+        public bool UseBlobIndexForSlowExists { get; set; } = false;
+        public bool IsPublicNamespace { get; set; } = true;
+        public NamespaceId? FallbackNamespace { get; set; } = null;
+        public bool PopulateFallbackNamespaceOnUpload { get; set; } = true;
+
+        public bool UseContentAddressedStorage { get; set; } = true;
+
+        public enum StoragePoolGCMethod
+        {
+            /// <summary>
+            /// Never run GC on this namespace
+            /// </summary>
+            None,
+            /// <summary>
+            /// Apply last access deletion, objects not used for a duration set in GCSettings will be removed
+            /// </summary>
+            LastAccess,
+            /// <summary>
+            /// Objects are removed after DefaultTTL time has passed, no matter if they are used or not
+            /// </summary>
+            TTL,
+            /// <summary>
+            /// Always GC references to this namespace, used to opt in to cleaning out old data
+            /// </summary>
+            Always
+        };
+
+        public StoragePoolGCMethod GcMethod { get; set; } = StoragePoolGCMethod.LastAccess;
+
+        public TimeSpan DefaultTTL { get; set; } = TimeSpan.FromDays(14);
+        public bool AllowRedirectUris { get; set; } = false;
+    }
+
     public class NamespacePolicyResolver : INamespacePolicyResolver
     {
         private readonly IOptionsMonitor<NamespaceSettings> _namespaceSettings;
diff --git a/Engine/Source/Programs/UnrealCloudDDC/Jupiter/Implementation/Blob/BlobService.cs b/Engine/Source/Programs/UnrealCloudDDC/Jupiter/Implementation/Blob/BlobService.cs
new file mode 100644
--- /dev/null
+++ b/Engine/Source/Programs/UnrealCloudDDC/Jupiter/Implementation/Blob/BlobService.cs
@@ -0,0 +1,949 @@
+﻿// Copyright Epic Games, Inc. All Rights Reserved.
+
+using System;
+using System.Collections.Concurrent;
+using System.Collections.Generic;
+using System.Diagnostics.Metrics;
+using System.IO;
+using System.Linq;
+using System.Net;
+using System.Net.Http;
+using System.Net.Mime;
+using System.Security.Claims;
+using System.Threading.Tasks;
+using EpicGames.AspNet;
+using EpicGames.Horde.Storage;
+using Jupiter.Controllers;
+using Jupiter.Implementation.Blob;
+using Jupiter.Common;
+using Jupiter.Common.Implementation;
+using Jupiter.Utils;
+using Microsoft.AspNetCore.Http;
+using Microsoft.AspNetCore.Mvc;
+using Microsoft.Extensions.DependencyInjection;
+using Microsoft.Extensions.Options;
+using Microsoft.Extensions.Logging;
+using OpenTelemetry.Trace;
+
+namespace Jupiter.Implementation;
+
+public class BlobService : IBlobService
+{
+    private List<IBlobStore> _blobStores;
+    private readonly IOptionsMonitor<UnrealCloudDDCSettings> _settings;
+    private readonly IBlobIndex _blobIndex;
+    private readonly IPeerStatusService _peerStatusService;
+    private readonly IHttpClientFactory _httpClientFactory;
+    private readonly IServiceCredentials _serviceCredentials;
+    private readonly INamespacePolicyResolver _namespacePolicyResolver;
+    private readonly IHttpContextAccessor _httpContextAccessor;
+    private readonly RequestHelper? _requestHelper;
+    private readonly Tracer _tracer;
+    private readonly BufferedPayloadFactory _bufferedPayloadFactory;
+    private readonly ILogger _logger;
+    private readonly Counter<long>? _storeGetHitsCounter;
+    private readonly Counter<long>? _storeGetAttemptsCounter;
+
+    internal IEnumerable<IBlobStore> BlobStore
+    {
+        get => _blobStores;
+        set => _blobStores = value.ToList();
+    }
+
+    public BlobService(IServiceProvider provider, IOptionsMonitor<UnrealCloudDDCSettings> settings, IBlobIndex blobIndex, IPeerStatusService peerStatusService, IHttpClientFactory httpClientFactory, IServiceCredentials serviceCredentials, INamespacePolicyResolver namespacePolicyResolver, IHttpContextAccessor httpContextAccessor, RequestHelper? requestHelper, Tracer tracer, BufferedPayloadFactory bufferedPayloadFactory, ILogger<BlobService> logger, Meter? meter)
+    {
+        _blobStores = GetBlobStores(provider, settings).ToList();
+        _settings = settings;
+        _blobIndex = blobIndex;
+        _peerStatusService = peerStatusService;
+        _httpClientFactory = httpClientFactory;
+        _serviceCredentials = serviceCredentials;
+        _namespacePolicyResolver = namespacePolicyResolver;
+        _httpContextAccessor = httpContextAccessor;
+        _requestHelper = requestHelper;
+        _tracer = tracer;
+        _bufferedPayloadFactory = bufferedPayloadFactory;
+        _logger = logger;
+
+        _storeGetHitsCounter = meter?.CreateCounter<long>("store.blob_get.found");
+        _storeGetAttemptsCounter = meter?.CreateCounter<long>("store.blob_get.attempt");
+    }
+
+    public static IEnumerable<IBlobStore> GetBlobStores(IServiceProvider provider, IOptionsMonitor<UnrealCloudDDCSettings> settings)
+    {
+        return settings.CurrentValue.GetStorageImplementations().Select(impl => ToStorageImplementation(provider, impl));
+    }
+
+    private static IBlobStore ToStorageImplementation(IServiceProvider provider, UnrealCloudDDCSettings.StorageBackendImplementations impl)
+    {
+        IBlobStore? store = impl switch
+        {
+            UnrealCloudDDCSettings.StorageBackendImplementations.S3 => provider.GetService<AmazonS3Store>(),
+            UnrealCloudDDCSettings.StorageBackendImplementations.Azure => provider.GetService<AzureBlobStore>(),
+            UnrealCloudDDCSettings.StorageBackendImplementations.FileSystem => provider.GetService<FileSystemStore>(),
+            UnrealCloudDDCSettings.StorageBackendImplementations.Memory => provider.GetService<MemoryBlobStore>(),
+            UnrealCloudDDCSettings.StorageBackendImplementations.Relay => provider.GetService<RelayBlobStore>(),
+            UnrealCloudDDCSettings.StorageBackendImplementations.Peer => provider.GetService<PeerBlobStore>(),
+            _ => throw new NotImplementedException("Unknown blob store {store")
+        };
+        if (store == null)
+        {
+            throw new ArgumentException($"Failed to find a provider service for type: {impl}");
+        }
+
+        return store;
+    }
+
+    public async Task<ContentHash> VerifyContentMatchesHash(Stream content, ContentHash identifier)
+    {
+        ContentHash blobHash;
+        {
+            using TelemetrySpan _ = _tracer.StartActiveSpan("web.hash").SetAttribute("operation.name", "web.hash");
+            blobHash = await BlobIdentifier.FromStream(content);
+        }
+
+        if (!identifier.Equals(blobHash))
+        {
+            throw new HashMismatchException(identifier, blobHash);
+        }
+
+        return identifier;
+    }
+
+    public async Task<BlobIdentifier> PutObject(NamespaceId ns, IBufferedPayload payload, BlobIdentifier identifier)
+    {
+        bool useContentAddressedStorage = _namespacePolicyResolver.GetPoliciesForNs(ns).UseContentAddressedStorage;
+        using TelemetrySpan scope = _tracer.StartActiveSpan("put_blob")
+            .SetAttribute("operation.name", "put_blob")
+            .SetAttribute("resource.name", identifier.ToString())
+            .SetAttribute("Content-Length", payload.Length.ToString());
+
+        await using Stream hashStream = payload.GetStream();
+        BlobIdentifier id = useContentAddressedStorage ? BlobIdentifier.FromContentHash(await VerifyContentMatchesHash(hashStream, identifier)) : identifier;
+
+        BlobIdentifier objectStoreIdentifier = await PutObjectToStores(ns, payload, id);
+        await _blobIndex.AddBlobToIndex(ns, id);
+
+        return objectStoreIdentifier;
+
+    }
+
+    public async Task<BlobIdentifier> PutObject(NamespaceId ns, byte[] payload, BlobIdentifier identifier)
+    {
+        bool useContentAddressedStorage = _namespacePolicyResolver.GetPoliciesForNs(ns).UseContentAddressedStorage;
+        using TelemetrySpan scope = _tracer.StartActiveSpan("put_blob")
+            .SetAttribute("operation.name", "put_blob")
+            .SetAttribute("resource.name", identifier.ToString())
+            .SetAttribute("Content-Length", payload.Length.ToString())
+            ;
+
+        await using Stream hashStream = new MemoryStream(payload);
+        BlobIdentifier id = useContentAddressedStorage ? BlobIdentifier.FromContentHash(await VerifyContentMatchesHash(hashStream, identifier)) : identifier;
+
+        BlobIdentifier objectStoreIdentifier = await PutObjectToStores(ns, payload, id);
+        await _blobIndex.AddBlobToIndex(ns, id);
+
+        return objectStoreIdentifier;
+    }
+
+    public async Task<Uri?> MaybePutObjectWithRedirect(NamespaceId ns, BlobIdentifier identifier)
+    {
+        bool allowRedirectUris = _namespacePolicyResolver.GetPoliciesForNs(ns).AllowRedirectUris;
+        if (!allowRedirectUris)
+        {
+            return null;
+        }
+
+        IServerTiming? serverTiming = _httpContextAccessor.HttpContext?.RequestServices.GetService<IServerTiming>();
+
+        // we only attempt to upload to the last store
+        // assuming that any other store will pull from it when they lack content once the upload has finished
+        IBlobStore store = _blobStores.Last();
+        {
+            string storeName = store.GetType().Name;
+            using TelemetrySpan scope = _tracer.StartActiveSpan("put_blob_with_redirect")
+                    .SetAttribute("operation.name", "put_blob_with_redirect")
+                    .SetAttribute("resource.name", identifier.ToString())
+                    .SetAttribute("store", storeName)
+            ;
+
+            using ServerTimingMetricScoped? serverTimingScope = serverTiming?.CreateServerTimingMetricScope($"blob.put-redirect.{storeName}", $"PUT(redirect) to store: '{storeName}'");
+
+            Uri? redirectUri = await store.PutObjectWithRedirect(ns, identifier);
+            if (redirectUri != null)
+            {
+                return redirectUri;
+            }
+        }
+
+        // no store found that supports redirect
+        return null;
+    }
+
+    public async Task<BlobIdentifier> PutObjectKnownHash(NamespaceId ns, IBufferedPayload content, BlobIdentifier identifier)
+    {
+        using TelemetrySpan scope = _tracer.StartActiveSpan("put_blob")
+            .SetAttribute("operation.name", "put_blob")
+            .SetAttribute("resource.name", identifier.ToString())
+            .SetAttribute("Content-Length", content.Length.ToString())
+            ;
+
+        BlobIdentifier objectStoreIdentifier = await PutObjectToStores(ns, content, identifier);
+        await _blobIndex.AddBlobToIndex(ns, identifier);
+
+        return objectStoreIdentifier;
+    }
+
+    private async Task<BlobIdentifier> PutObjectToStores(NamespaceId ns, IBufferedPayload bufferedPayload, BlobIdentifier identifier)
+    {
+        IServerTiming? serverTiming = _httpContextAccessor.HttpContext?.RequestServices.GetService<IServerTiming>();
+
+        foreach (IBlobStore store in _blobStores)
+        {
+            using TelemetrySpan scope = _tracer.StartActiveSpan("put_blob_to_store")
+                .SetAttribute("operation.name", "put_blob_to_store")
+                .SetAttribute("resource.name", identifier.ToString())
+                .SetAttribute("store", store.GetType().ToString())
+                ;
+            string storeName = store.GetType().Name;
+
+            using ServerTimingMetricScoped? serverTimingScope = serverTiming?.CreateServerTimingMetricScope($"blob.put.{storeName}", $"PUT to store: '{storeName}'");
+
+            await using Stream s = bufferedPayload.GetStream();
+            await store.PutObject(ns, s, identifier);
+        }
+        NamespacePolicy policy = _namespacePolicyResolver.GetPoliciesForNs(ns);
+        if (policy.PopulateFallbackNamespaceOnUpload && policy.FallbackNamespace.HasValue)
+        {
+            await PutObjectToStores(policy.FallbackNamespace.Value, bufferedPayload, identifier);
+            await _blobIndex.AddBlobToIndex(policy.FallbackNamespace.Value, identifier);
+        }
+        return identifier;
+    }
+
+    private async Task<BlobIdentifier> PutObjectToStores(NamespaceId ns, byte[] payload, BlobIdentifier identifier)
+    {
+        foreach (IBlobStore store in _blobStores)
+        {
+            await store.PutObject(ns, payload, identifier);
+        }
+
+        NamespacePolicy policy = _namespacePolicyResolver.GetPoliciesForNs(ns);
+        if (policy.PopulateFallbackNamespaceOnUpload && policy.FallbackNamespace.HasValue)
+        {
+            await PutObjectToStores(policy.FallbackNamespace.Value, payload, identifier);
+            await _blobIndex.AddBlobToIndex(policy.FallbackNamespace.Value, identifier);
+        }
+        return identifier;
+    }
+
+    public async Task<BlobContents> GetObject(NamespaceId ns, BlobIdentifier blob, List<string>? storageLayers = null, bool supportsRedirectUri = false)
+    {
+        try
+        {
+            return await GetObjectFromStores(ns, blob, storageLayers, supportsRedirectUri);
+        }
+        catch (BlobNotFoundException)
+        {
+            NamespacePolicy policy = _namespacePolicyResolver.GetPoliciesForNs(ns);
+
+            if (policy.FallbackNamespace != null)
+            {
+                ClaimsPrincipal? user = _httpContextAccessor.HttpContext?.User;
+                HttpRequest? request = _httpContextAccessor.HttpContext?.Request;
+
+                if (user == null || request == null)
+                {
+                    _logger.LogWarning("Unable to fallback to namespace {Namespace} due to not finding required http context values", policy.FallbackNamespace.Value);
+                    throw;
+                }
+
+                if (_requestHelper == null)
+                {
+                    _logger.LogWarning("Unable to fallback to namespace {Namespace} due to request helper missing", policy.FallbackNamespace.Value);
+                    throw;
+                }
+
+                ActionResult? result = await _requestHelper.HasAccessToNamespace(user, request, policy.FallbackNamespace.Value, new [] { AclAction.ReadObject });
+                if (result != null)
+                {
+                    _logger.LogInformation("Authorization error when attempting to fallback to namespace {FallbackNamespace}. This may be confusing for users that as they had access to original namespace {Namespace}", policy.FallbackNamespace.Value, ns);
+                    throw new AuthorizationException(result, "Failed to authenticate for fallback namespace");
+                }
+
+                try
+                {
+                    IServerTiming? serverTiming = _httpContextAccessor.HttpContext?.RequestServices.GetService<IServerTiming>();
+
+                    // read the content from the fallback namespace
+                    BlobContents fallbackContent = await GetObjectFromStores(policy.FallbackNamespace.Value, blob, storageLayers);
+                    
+                    // populate the primary namespace with the content
+                    using TelemetrySpan _ = _tracer.StartActiveSpan("HierarchicalStore.Populate").SetAttribute("operation.name", "HierarchicalStore.Populate");
+                    using ServerTimingMetricScoped? serverTimingScope = serverTiming?.CreateServerTimingMetricScope($"blob.populate", "Populating caches with blob contents");
+
+                    await using MemoryStream tempStream = new MemoryStream();
+                    await fallbackContent.Stream.CopyToAsync(tempStream);
+                    byte[] data = tempStream.ToArray();
+
+                    await PutObject(ns, data, blob);
+                    return await GetObject(ns, blob);
+                }
+                catch (BlobNotFoundException)
+                {
+                    // if we fail to find the blob in the fallback namespace we can carry on as we will rethrow the blob not found exception later (if the replication also fails)
+                }
+            }
+
+            if (ShouldFetchBlobOnDemand(ns))
+            {
+                try
+                {
+                    return await ReplicateObject(ns, blob);
+                }
+                catch (BlobNotFoundException)
+                {
+                    // if the blob is not found we can ignore it as we will just rethrow it later anyway
+                }
+            }
+
+            // if the primary namespace failed check to see if we should use a fallback policy which has replication enabled
+            if (policy.FallbackNamespace != null && ShouldFetchBlobOnDemand(policy.FallbackNamespace.Value))
+            {
+                try
+                {
+                    return await ReplicateObject(policy.FallbackNamespace.Value, blob);
+                }
+                catch (BlobNotFoundException)
+                {
+                    // if the blob is not found we can ignore it as we will just rethrow it later anyway
+                }
+            }
+
+            // we might have attempted to fetch the object and failed, or had no fallback options, either way the blob can not be found so we should rethrow
+            throw;
+        }
+    }
+
+    private async Task<BlobContents> GetObjectFromStores(NamespaceId ns, BlobIdentifier blob, List<string>? storageLayers = null, bool supportsRedirectUri = false)
+    {
+        bool seenBlobNotFound = false;
+        bool seenNamespaceNotFound = false;
+        int numStoreMisses = 0;
+        BlobContents? blobContents = null;
+        IServerTiming? serverTiming = _httpContextAccessor.HttpContext?.RequestServices.GetService<IServerTiming>();
+
+        foreach (IBlobStore store in _blobStores)
+        {
+            string blobStoreName = store.GetType().Name;
+
+            // check which storage layers to skip if we have a explicit list of storage layers to use
+            if (storageLayers != null && storageLayers.Count != 0)
+            {
+                bool found = false;
+                foreach (string storageLayer in storageLayers)
+                {
+                    if (string.Equals(storageLayer, blobStoreName, StringComparison.OrdinalIgnoreCase))
+                    {
+                        found = true;
+                        break;
+                    }
+                }
+
+                if (!found)
+                {
+                    continue;
+                }
+            }
+
+            using TelemetrySpan scope = _tracer.StartActiveSpan("HierarchicalStore.GetObject")
+                .SetAttribute("operation.name", "HierarchicalStore.GetObject")
+                .SetAttribute("resource.name",  blob.ToString())
+                .SetAttribute("BlobStore", store.GetType().ToString())
+                .SetAttribute("ObjectFound", false.ToString())
+                ;
+
+            string storeName = store.GetType().Name;
+            using ServerTimingMetricScoped? serverTimingScope = serverTiming?.CreateServerTimingMetricScope($"blob.get.{storeName}", $"Blob GET from: '{storeName}'");
+            _storeGetAttemptsCounter?.Add(1, new KeyValuePair<string, object?>("store", storeName));
+            try
+            {
+                blobContents = await store.GetObject(ns, blob, supportsRedirectUri: supportsRedirectUri);
+                scope.SetAttribute("ObjectFound", true.ToString());
+                _storeGetHitsCounter?.Add(1, new KeyValuePair<string, object?>("store", storeName));
+                break;
+            }
+            catch (BlobNotFoundException)
+            {
+                seenBlobNotFound = true;
+                numStoreMisses++;
+            }
+            catch (NamespaceNotFoundException)
+            {
+                seenNamespaceNotFound = true;
+            }
+        }
+
+        if (seenBlobNotFound && blobContents == null)
+        {
+            throw new BlobNotFoundException(ns, blob);
+        }
+
+        if (seenNamespaceNotFound && blobContents == null)
+        {
+            throw new NamespaceNotFoundException(ns);
+        }
+
+        // if we applied filters to the storage layers resulting in no blob found we consider it a miss
+        if (storageLayers != null && storageLayers.Count != 0 && blobContents == null)
+        {
+            throw new BlobNotFoundException(ns, blob);
+        }
+
+        if (blobContents == null)
+        {
+            // Should not happen but exists to safeguard against the null pointer
+            throw new Exception("blobContents is null");
+        }
+
+        if (numStoreMisses >= 1)
+        {
+            using TelemetrySpan _ = _tracer.StartActiveSpan("HierarchicalStore.Populate").SetAttribute("operation.name", "HierarchicalStore.Populate");
+            using ServerTimingMetricScoped? serverTimingScope = serverTiming?.CreateServerTimingMetricScope($"blob.populate", "Populating caches with blob contents");
+
+            await using MemoryStream tempStream = new MemoryStream();
+            await blobContents.Stream.CopyToAsync(tempStream);
+            byte[] data = tempStream.ToArray();
+            
+            // Don't populate the last store, as that is where we got the hit
+            for (int i = 0; i < numStoreMisses; i++)
+            {
+                IBlobStore blobStore = _blobStores[i];
+                using TelemetrySpan scope = _tracer.StartActiveSpan("HierarchicalStore.PopulateStore")
+                    .SetAttribute("operation.name", "HierarchicalStore.PopulateStore")
+                    .SetAttribute("BlobStore", blobStore.GetType().Name);
+                // Populate each store traversed that did not have the content found lower in the hierarchy
+                await blobStore.PutObject(ns, data, blob);
+            }
+
+#pragma warning disable CA2000 // Dispose objects before losing scope , ownership is transfered to caller
+            blobContents = new BlobContents(data);
+#pragma warning restore CA2000 // Dispose objects before losing scope
+        }
+        
+        return blobContents;
+    }
+
+    public async Task<Uri?> GetObjectWithRedirect(NamespaceId ns, BlobIdentifier blob, List<string>? storageLayers = null)
+    {
+        bool seenBlobNotFound = false;
+        bool seenNamespaceNotFound = false;
+        int numStoreMisses = 0;
+        IServerTiming? serverTiming = _httpContextAccessor.HttpContext?.RequestServices.GetService<IServerTiming>();
+
+        Uri? redirectUri = null;
+        foreach (IBlobStore store in _blobStores)
+        {
+            string blobStoreName = store.GetType().Name;
+
+            // check which storage layers to skip if we have a explicit list of storage layers to use
+            if (storageLayers != null && storageLayers.Count != 0)
+            {
+                bool found = false;
+                foreach (string storageLayer in storageLayers)
+                {
+                    if (string.Equals(storageLayer, blobStoreName, StringComparison.OrdinalIgnoreCase))
+                    {
+                        found = true;
+                        break;
+                    }
+                }
+
+                if (!found)
+                {
+                    continue;
+                }
+            }
+
+            using TelemetrySpan scope = _tracer.StartActiveSpan("HierarchicalStore.GetObjectRedirect")
+                .SetAttribute("operation.name", "HierarchicalStore.GetObject")
+                .SetAttribute("resource.name",  blob.ToString())
+                .SetAttribute("BlobStore", store.GetType().ToString())
+                .SetAttribute("ObjectFound", false.ToString())
+                ;
+
+            string storeName = store.GetType().Name;
+            using ServerTimingMetricScoped? serverTimingScope = serverTiming?.CreateServerTimingMetricScope($"blob.get-redirect.{storeName}", $"Blob GET from: '{storeName}'");
+
+            try
+            {
+                redirectUri = await store.GetObjectByRedirect(ns, blob);
+                scope.SetAttribute("ObjectFound", true.ToString());
+                break;
+            }
+            catch (BlobNotFoundException)
+            {
+                seenBlobNotFound = true;
+                numStoreMisses++;
+            }
+            catch (NamespaceNotFoundException)
+            {
+                seenNamespaceNotFound = true;
+            }
+        }
+
+        if (seenBlobNotFound && redirectUri == null)
+        {
+            throw new BlobNotFoundException(ns, blob);
+        }
+
+        if (seenNamespaceNotFound && redirectUri == null)
+        {
+            throw new NamespaceNotFoundException(ns);
+        }
+
+        // if we applied filters to the storage layers resulting in no blob found we consider it a miss
+        if (storageLayers != null && storageLayers.Count != 0 && redirectUri == null)
+        {
+            throw new BlobNotFoundException(ns, blob);
+        }
+
+        return redirectUri;
+    }
+
+    public async Task<BlobContents> ReplicateObject(NamespaceId ns, BlobIdentifier blob, bool force = false)
+    {
+        if (!force && !ShouldFetchBlobOnDemand(ns))
+        {
+            throw new NotSupportedException($"Replication is not allowed in namespace {ns}");
+        }
+        IServerTiming? serverTiming = _httpContextAccessor.HttpContext?.RequestServices.GetService<IServerTiming>();
+        using TelemetrySpan scope = _tracer.StartActiveSpan("HierarchicalStore.Replicate").SetAttribute("operation.name", "HierarchicalStore.Replicate");
+
+        using ServerTimingMetricScoped? serverTimingScope = serverTiming?.CreateServerTimingMetricScope("blob.replicate", "Replicating blob from remote instances");
+
+        List<string> regions = await _blobIndex.GetBlobRegions(ns, blob);
+
+        if (!regions.Any())
+        {
+            throw new BlobReplicationException(ns, blob, "Blob not found in any region");
+        }
+
+        _logger.LogInformation("On-demand replicating blob {Blob} in Namespace {Namespace}", blob, ns);
+        List<(int, string)> possiblePeers = new List<(int, string)>(_peerStatusService.GetPeersByLatency(regions));
+
+        bool replicated = false;
+        foreach ((int latency, string? region) in possiblePeers)
+        {
+            PeerStatus? peerStatus = _peerStatusService.GetPeerStatus(region);
+            if (peerStatus == null)
+            {
+                throw new Exception($"Failed to find peer {region}");
+            }
+
+            _logger.LogInformation("Attempting to replicate blob {Blob} in Namespace {Namespace} from {Region}", blob, ns, region);
+
+            PeerEndpoints peerEndpoint = peerStatus.Endpoints.First();
+            using HttpClient httpClient = _httpClientFactory.CreateClient();
+            string url = peerEndpoint.Url.ToString();
+            if (!url.EndsWith("/", StringComparison.InvariantCultureIgnoreCase))
+            {
+                url += "/";
+            }
+            using HttpRequestMessage blobRequest = await BuildHttpRequest(HttpMethod.Get, new Uri($"{url}api/v1/blobs/{ns}/{blob}"));
+            HttpResponseMessage blobResponse = await httpClient.SendAsync(blobRequest, HttpCompletionOption.ResponseHeadersRead);
+
+            if (blobResponse.StatusCode == HttpStatusCode.NotFound)
+            {
+                // try the next region
+                continue;
+            }
+
+            if (blobResponse.StatusCode != HttpStatusCode.OK)
+            {
+                throw new BlobReplicationException(ns, blob, $"Failed to replicate {blob} in {ns} from region {region} due to bad http status code {blobResponse.StatusCode}.");
+            }
+
+            await using Stream s = await blobResponse.Content.ReadAsStreamAsync();
+            using IBufferedPayload payload = await _bufferedPayloadFactory.CreateFilesystemBufferedPayload(s);
+            await PutObject(ns, payload, blob);
+            replicated = true;
+            break;
+        }
+
+        if (!replicated)
+        {
+            throw new BlobReplicationException(ns, blob, $"Failed to replicate {blob} in {ns} due to it not existing in any region");
+        }
+
+        return await GetObject(ns, blob);
+    }
+
+    public bool ShouldFetchBlobOnDemand(NamespaceId ns)
+    {
+        return _settings.CurrentValue.EnableOnDemandReplication && _namespacePolicyResolver.GetPoliciesForNs(ns).OnDemandReplication;
+    }
+
+    private async Task<HttpRequestMessage> BuildHttpRequest(HttpMethod httpMethod, Uri uri)
+    {
+        string? token = await _serviceCredentials.GetTokenAsync();
+        HttpRequestMessage request = new HttpRequestMessage(httpMethod, uri);
+        if (!string.IsNullOrEmpty(token))
+        {
+            request.Headers.Add("Authorization", $"{_serviceCredentials.GetAuthenticationScheme()} {token}");
+        }
+
+        return request;
+    }
+
+    public async Task<bool> Exists(NamespaceId ns, BlobIdentifier blob, List<string>? storageLayers = null)
+    {
+        bool exists = await ExistsInStores(ns, blob, storageLayers);
+        if (exists)
+        {
+            return exists;
+        }
+
+        NamespacePolicy policy = _namespacePolicyResolver.GetPoliciesForNs(ns);
+        if (policy.FallbackNamespace != null)
+        {
+            return await ExistsInStores(policy.FallbackNamespace.Value, blob, storageLayers);
+        }
+
+        return false;
+    }
+
+    private async Task<bool> ExistsInStores(NamespaceId ns, BlobIdentifier blob, List<string>? storageLayers = null)
+    {
+        bool useBlobIndex = _namespacePolicyResolver.GetPoliciesForNs(ns).UseBlobIndexForExists;
+        if (useBlobIndex)
+        {
+            using TelemetrySpan scope = _tracer.StartActiveSpan("HierarchicalStore.ObjectExists")
+                .SetAttribute("operation.name", "HierarchicalStore.ObjectExists")
+                .SetAttribute("resource.name", blob.ToString())
+                .SetAttribute("BlobStore", "BlobIndex")
+                ;
+            bool exists = await _blobIndex.BlobExistsInRegion(ns, blob);
+            if (exists)
+            {
+                scope.SetAttribute("ObjectFound", true.ToString());
+                return true;
+            }
+
+            scope.SetAttribute("ObjectFound", false.ToString());
+            return false;
+        }
+        else
+        {
+            foreach (IBlobStore store in _blobStores)
+            {
+                string blobStoreName = store.GetType().Name;
+                // check which storage layers to skip if we have a explicit list of storage layers to use
+                if (storageLayers != null && storageLayers.Count != 0)
+                {
+                    bool found = false;
+                    foreach (string storageLayer in storageLayers)
+                    {
+                        if (string.Equals(storageLayer, blobStoreName, StringComparison.OrdinalIgnoreCase))
+                        {
+                            found = true;
+                            break;
+                        }
+                    }
+
+                    if (!found)
+                    {
+                        continue;
+                    }
+                }
+
+                using TelemetrySpan scope = _tracer.StartActiveSpan("HierarchicalStore.ObjectExists")
+                    .SetAttribute("operation.name", "HierarchicalStore.ObjectExists")
+                    .SetAttribute("resource.name", blob.ToString())
+                    .SetAttribute("BlobStore", blobStoreName)
+                    ;
+                if (await store.Exists(ns, blob))
+                {
+                    scope.SetAttribute("ObjectFound", true.ToString());
+                    return true;
+                }
+                scope.SetAttribute("ObjectFound", false.ToString());
+            }
+
+            return false;
+        }
+    }
+
+    public async Task<bool> ExistsInRootStore(NamespaceId ns, BlobIdentifier blob)
+    {
+        IBlobStore store = _blobStores.Last();
+
+        using TelemetrySpan scope = _tracer.StartActiveSpan("HierarchicalStore.ObjectExistsInRoot")
+            .SetAttribute("operation.name", "HierarchicalStore.ObjectExistsInRoot")
+            .SetAttribute("resource.name", blob.ToString())
+            .SetAttribute("BlobStore", store.GetType().Name)
+            ;
+        if (await store.Exists(ns, blob))
+        {
+            scope.SetAttribute("ObjectFound", true.ToString());
+            return true;
+        }
+        scope.SetAttribute("ObjectFound", false.ToString());
+        return false;
+    }
+
+    public async Task DeleteObject(NamespaceId ns, BlobIdentifier blob)
+    {
+        bool blobNotFound = false;
+        bool deletedAtLeastOnce = false;
+
+        // remove the object from the tracking first, if this times out we do not want to end up with a inconsistent blob index
+        // if the blob store delete fails on the other hand we will still run a delete again during GC (as the blob is still orphaned at that point)
+        // this assumes that blob gc is based on scanning the root blob store
+        await _blobIndex.RemoveBlobFromRegion(ns, blob);
+
+        foreach (IBlobStore store in _blobStores)
+        {
+            try
+            {
+                using TelemetrySpan scope = _tracer.StartActiveSpan("HierarchicalStore.DeleteObject")
+                    .SetAttribute("operation.name", "HierarchicalStore.DeleteObject")
+                    .SetAttribute("resource.name", blob.ToString())
+                    .SetAttribute("BlobStore", store.GetType().Name)
+                    ;
+
+                await store.DeleteObject(ns, blob);
+                deletedAtLeastOnce = true;
+            }
+            catch (NamespaceNotFoundException)
+            {
+                // Ignore
+            }
+            catch (BlobNotFoundException)
+            {
+                blobNotFound = true;
+            }
+        }
+
+        if (deletedAtLeastOnce)
+        {
+            return;
+        }
+
+        if (blobNotFound)
+        {
+            throw new BlobNotFoundException(ns, blob);
+        }
+
+        throw new NamespaceNotFoundException(ns);
+    }
+
+    public async Task DeleteNamespace(NamespaceId ns)
+    {
+        bool deletedAtLeastOnce = false;
+        foreach (IBlobStore store in _blobStores)
+        {
+            using TelemetrySpan scope = _tracer.StartActiveSpan("HierarchicalStore.DeleteNamespace")
+                .SetAttribute("operation.name", "HierarchicalStore.DeleteNamespace")
+                .SetAttribute("resource.name", ns.ToString())
+                .SetAttribute("BlobStore", store.GetType().Name)
+                ;
+            try
+            {
+                await store.DeleteNamespace(ns);
+                deletedAtLeastOnce = true;
+            }
+            catch (NamespaceNotFoundException)
+            {
+                // Ignore
+            }
+        }
+
+        if (deletedAtLeastOnce)
+        {
+            return;
+        }
+
+        throw new NamespaceNotFoundException(ns);
+    }
+
+    public IAsyncEnumerable<(BlobIdentifier,DateTime)> ListObjects(NamespaceId ns)
+    {
+        // as this is a hierarchy of blob stores the last blob store should contain the superset of all stores
+        return _blobStores.Last().ListObjects(ns);
+    }
+
+    public async Task<BlobIdentifier[]> FilterOutKnownBlobs(NamespaceId ns, IEnumerable<BlobIdentifier> blobs)
+    {
+        List<(BlobIdentifier, Task<bool>)> existTasks = new();
+        foreach (BlobIdentifier blob in blobs)
+        {
+            existTasks.Add((blob, Exists(ns, blob)));
+        }
+
+        List<BlobIdentifier> missingBlobs = new();
+        foreach ((BlobIdentifier blob, Task<bool> existsTask) in existTasks)
+        {
+            bool exists = await existsTask;
+
+            if (!exists)
+            {
+                missingBlobs.Add(blob);
+            }
+        }
+
+        return missingBlobs.ToArray();
+    }
+
+    public async Task<BlobIdentifier[]> FilterOutKnownBlobs(NamespaceId ns, IAsyncEnumerable<BlobIdentifier> blobs)
+    {
+        ConcurrentBag<BlobIdentifier> missingBlobs = new ConcurrentBag<BlobIdentifier>();
+
+        try
+        {
+            await Parallel.ForEachAsync(blobs, async (identifier, ctx) =>
+            {
+                bool exists = await Exists(ns, identifier);
+
+                if (!exists)
+                {
+                    missingBlobs.Add(identifier);
+                }
+            });
+        }
+        catch (AggregateException e)
+        {
+            if (e.InnerException is PartialReferenceResolveException)
+            {
+                throw e.InnerException;
+            }
+
+            throw;
+        }
+
+        return missingBlobs.ToArray();
+    }
+
+    public async Task<BlobContents> GetObjects(NamespaceId ns, BlobIdentifier[] blobs)
+    {
+        using TelemetrySpan _ = _tracer.StartActiveSpan("blob.combine").SetAttribute("operation.name", "blob.combine");
+        Task<BlobContents>[] tasks = new Task<BlobContents>[blobs.Length];
+        for (int i = 0; i < blobs.Length; i++)
+        {
+            tasks[i] = GetObject(ns, blobs[i]);
+        }
+
+        MemoryStream ms = new MemoryStream();
+        foreach (Task<BlobContents> task in tasks)
+        {
+            BlobContents blob = await task;
+            await using Stream s = blob.Stream;
+            await s.CopyToAsync(ms);
+        }
+
+        ms.Seek(0, SeekOrigin.Begin);
+
+        return new BlobContents(ms, ms.Length);
+    }
+}
+
+public static class BlobServiceExtensions
+{
+    public static async Task<ContentId> PutCompressedObject(this IBlobService blobService, NamespaceId ns, IBufferedPayload payload, ContentId? id, IServiceProvider provider)
+    {
+        IContentIdStore contentIdStore = provider.GetService<IContentIdStore>()!;
+        CompressedBufferUtils compressedBufferUtils = provider.GetService<CompressedBufferUtils>()!;
+        Tracer tracer = provider.GetService<Tracer>()!;
+
+        // decompress the content and generate a identifier from it to verify the identifier we got
+        await using Stream decompressStream = payload.GetStream();
+        // TODO: we should add a overload for decompress content that can work on streams, otherwise we are still limited to 2GB compressed blobs
+        byte[] decompressedContent = compressedBufferUtils.DecompressContent(await decompressStream.ToByteArray());
+
+        await using MemoryStream decompressedStream = new MemoryStream(decompressedContent);
+        ContentId identifierDecompressedPayload;
+        if (id != null)
+        {
+            identifierDecompressedPayload = ContentId.FromContentHash(await blobService.VerifyContentMatchesHash(decompressedStream, id));
+        }
+        else
+        {
+            ContentHash blobHash;
+            {
+                using TelemetrySpan _ = tracer.StartActiveSpan("web.hash").SetAttribute("operation.name", "web.hash");
+                blobHash = await BlobIdentifier.FromStream(decompressedStream);
+            }
+
+            identifierDecompressedPayload = ContentId.FromContentHash(blobHash);
+        }
+
+        BlobIdentifier identifierCompressedPayload;
+        {
+            using TelemetrySpan _ = tracer.StartActiveSpan("web.hash").SetAttribute("operation.name", "web.hash");
+            await using Stream hashStream = payload.GetStream();
+            identifierCompressedPayload = await BlobIdentifier.FromStream(hashStream);
+        }
+
+        // commit the mapping from the decompressed hash to the compressed hash, we run this in parallel with the blob store submit
+        // TODO: let users specify weight of the blob compared to previously submitted content ids
+        int contentIdWeight = (int)payload.Length;
+        Task contentIdStoreTask = contentIdStore.Put(ns, identifierDecompressedPayload, identifierCompressedPayload, contentIdWeight);
+
+        // we still commit the compressed buffer to the object store using the hash of the compressed content
+        {
+            await blobService.PutObjectKnownHash(ns, payload, identifierCompressedPayload);
+        }
+
+        await contentIdStoreTask;
+
+        return identifierDecompressedPayload;
+    }
+
+    public static async Task<(BlobContents, string)> GetCompressedObject(this IBlobService blobService, NamespaceId ns, ContentId contentId, IServiceProvider provider, bool supportsRedirectUri = false)
+    {
+        IContentIdStore contentIdStore = provider.GetService<IContentIdStore>()!;
+        Tracer tracer = provider.GetService<Tracer>()!;
+
+        BlobIdentifier[]? chunks = await contentIdStore.Resolve(ns, contentId, mustBeContentId: false);
+        if (chunks == null || chunks.Length == 0)
+        {
+            throw new ContentIdResolveException(contentId);
+        }
+
+        // single chunk, we just return that chunk
+        if (chunks.Length == 1)
+        {
+            BlobIdentifier blobToReturn = chunks[0];
+            string mimeType = CustomMediaTypeNames.UnrealCompressedBuffer;
+            if (contentId.Equals(blobToReturn))
+            {
+                // this was actually the unmapped blob, meaning its not a compressed buffer
+                mimeType = MediaTypeNames.Application.Octet;
+            }
+
+            return (await blobService.GetObject(ns, blobToReturn, supportsRedirectUri: supportsRedirectUri), mimeType);
+        }
+
+        // chunked content, combine the chunks into a single stream
+        using TelemetrySpan _ = tracer.StartActiveSpan("blob.combine").SetAttribute("operation.name", "blob.combine");
+        Task<BlobContents>[] tasks = new Task<BlobContents>[chunks.Length];
+        for (int i = 0; i < chunks.Length; i++)
+        {
+            // even if it was requested to support redirect, since we need to combine the chunks using redirects is not possible
+            tasks[i] = blobService.GetObject(ns, chunks[i], supportsRedirectUri: false);
+        }
+
+        MemoryStream ms = new MemoryStream();
+        foreach (Task<BlobContents> task in tasks)
+        {
+            BlobContents blob = await task;
+            await using Stream s = blob.Stream;
+            await s.CopyToAsync(ms);
+        }
+
+        ms.Seek(0, SeekOrigin.Begin);
+
+        // chunking could not have happened for a non compressed buffer so assume it is compressed
+        return (new BlobContents(ms, ms.Length), CustomMediaTypeNames.UnrealCompressedBuffer);
+    }
+}
\ No newline at end of file
diff --git a/Engine/Source/Programs/UnrealCloudDDC/Jupiter/Implementation/Blob/IBlobService.cs b/Engine/Source/Programs/UnrealCloudDDC/Jupiter/Implementation/Blob/IBlobService.cs
--- a/Engine/Source/Programs/UnrealCloudDDC/Jupiter/Implementation/Blob/IBlobService.cs
+++ b/Engine/Source/Programs/UnrealCloudDDC/Jupiter/Implementation/Blob/IBlobService.cs
@@ -1,29 +1,11 @@
 ﻿// Copyright Epic Games, Inc. All Rights Reserved.
 
 using System;
-using System.Collections.Concurrent;
 using System.Collections.Generic;
-using System.Diagnostics.Metrics;
 using System.IO;
-using System.Linq;
-using System.Net;
-using System.Net.Http;
-using System.Net.Mime;
-using System.Security.Claims;
 using System.Threading.Tasks;
-using EpicGames.AspNet;
 using EpicGames.Horde.Storage;
-using Jupiter.Controllers;
-using Jupiter.Implementation.Blob;
-using Jupiter.Common;
 using Jupiter.Common.Implementation;
-using Jupiter.Utils;
-using Microsoft.AspNetCore.Http;
-using Microsoft.AspNetCore.Mvc;
-using Microsoft.Extensions.DependencyInjection;
-using Microsoft.Extensions.Options;
-using Microsoft.Extensions.Logging;
-using OpenTelemetry.Trace;
 
 namespace Jupiter.Implementation;
 
@@ -65,923 +47,3 @@ public interface IBlobService
     bool ShouldFetchBlobOnDemand(NamespaceId ns);
 }
 
-public class BlobService : IBlobService
-{
-    private List<IBlobStore> _blobStores;
-    private readonly IOptionsMonitor<UnrealCloudDDCSettings> _settings;
-    private readonly IBlobIndex _blobIndex;
-    private readonly IPeerStatusService _peerStatusService;
-    private readonly IHttpClientFactory _httpClientFactory;
-    private readonly IServiceCredentials _serviceCredentials;
-    private readonly INamespacePolicyResolver _namespacePolicyResolver;
-    private readonly IHttpContextAccessor _httpContextAccessor;
-    private readonly RequestHelper? _requestHelper;
-    private readonly Tracer _tracer;
-    private readonly BufferedPayloadFactory _bufferedPayloadFactory;
-    private readonly ILogger _logger;
-    private readonly Counter<long>? _storeGetHitsCounter;
-    private readonly Counter<long>? _storeGetAttemptsCounter;
-
-    internal IEnumerable<IBlobStore> BlobStore
-    {
-        get => _blobStores;
-        set => _blobStores = value.ToList();
-    }
-
-    public BlobService(IServiceProvider provider, IOptionsMonitor<UnrealCloudDDCSettings> settings, IBlobIndex blobIndex, IPeerStatusService peerStatusService, IHttpClientFactory httpClientFactory, IServiceCredentials serviceCredentials, INamespacePolicyResolver namespacePolicyResolver, IHttpContextAccessor httpContextAccessor, RequestHelper? requestHelper, Tracer tracer, BufferedPayloadFactory bufferedPayloadFactory, ILogger<BlobService> logger, Meter? meter)
-    {
-        _blobStores = GetBlobStores(provider, settings).ToList();
-        _settings = settings;
-        _blobIndex = blobIndex;
-        _peerStatusService = peerStatusService;
-        _httpClientFactory = httpClientFactory;
-        _serviceCredentials = serviceCredentials;
-        _namespacePolicyResolver = namespacePolicyResolver;
-        _httpContextAccessor = httpContextAccessor;
-        _requestHelper = requestHelper;
-        _tracer = tracer;
-        _bufferedPayloadFactory = bufferedPayloadFactory;
-        _logger = logger;
-
-        _storeGetHitsCounter = meter?.CreateCounter<long>("store.blob_get.found");
-        _storeGetAttemptsCounter = meter?.CreateCounter<long>("store.blob_get.attempt");
-    }
-
-    public static IEnumerable<IBlobStore> GetBlobStores(IServiceProvider provider, IOptionsMonitor<UnrealCloudDDCSettings> settings)
-    {
-        return settings.CurrentValue.GetStorageImplementations().Select(impl => ToStorageImplementation(provider, impl));
-    }
-
-    private static IBlobStore ToStorageImplementation(IServiceProvider provider, UnrealCloudDDCSettings.StorageBackendImplementations impl)
-    {
-        IBlobStore? store = impl switch
-        {
-            UnrealCloudDDCSettings.StorageBackendImplementations.S3 => provider.GetService<AmazonS3Store>(),
-            UnrealCloudDDCSettings.StorageBackendImplementations.Azure => provider.GetService<AzureBlobStore>(),
-            UnrealCloudDDCSettings.StorageBackendImplementations.FileSystem => provider.GetService<FileSystemStore>(),
-            UnrealCloudDDCSettings.StorageBackendImplementations.Memory => provider.GetService<MemoryBlobStore>(),
-            UnrealCloudDDCSettings.StorageBackendImplementations.Relay => provider.GetService<RelayBlobStore>(),
-            UnrealCloudDDCSettings.StorageBackendImplementations.Peer => provider.GetService<PeerBlobStore>(),
-            _ => throw new NotImplementedException("Unknown blob store {store")
-        };
-        if (store == null)
-        {
-            throw new ArgumentException($"Failed to find a provider service for type: {impl}");
-        }
-
-        return store;
-    }
-
-    public async Task<ContentHash> VerifyContentMatchesHash(Stream content, ContentHash identifier)
-    {
-        ContentHash blobHash;
-        {
-            using TelemetrySpan _ = _tracer.StartActiveSpan("web.hash").SetAttribute("operation.name", "web.hash");
-            blobHash = await BlobIdentifier.FromStream(content);
-        }
-
-        if (!identifier.Equals(blobHash))
-        {
-            throw new HashMismatchException(identifier, blobHash);
-        }
-
-        return identifier;
-    }
-
-    public async Task<BlobIdentifier> PutObject(NamespaceId ns, IBufferedPayload payload, BlobIdentifier identifier)
-    {
-        bool useContentAddressedStorage = _namespacePolicyResolver.GetPoliciesForNs(ns).UseContentAddressedStorage;
-        using TelemetrySpan scope = _tracer.StartActiveSpan("put_blob")
-            .SetAttribute("operation.name", "put_blob")
-            .SetAttribute("resource.name", identifier.ToString())
-            .SetAttribute("Content-Length", payload.Length.ToString());
-
-        await using Stream hashStream = payload.GetStream();
-        BlobIdentifier id = useContentAddressedStorage ? BlobIdentifier.FromContentHash(await VerifyContentMatchesHash(hashStream, identifier)) : identifier;
-
-        BlobIdentifier objectStoreIdentifier = await PutObjectToStores(ns, payload, id);
-        await _blobIndex.AddBlobToIndex(ns, id);
-
-        return objectStoreIdentifier;
-
-    }
-
-    public async Task<BlobIdentifier> PutObject(NamespaceId ns, byte[] payload, BlobIdentifier identifier)
-    {
-        bool useContentAddressedStorage = _namespacePolicyResolver.GetPoliciesForNs(ns).UseContentAddressedStorage;
-        using TelemetrySpan scope = _tracer.StartActiveSpan("put_blob")
-            .SetAttribute("operation.name", "put_blob")
-            .SetAttribute("resource.name", identifier.ToString())
-            .SetAttribute("Content-Length", payload.Length.ToString())
-            ;
-
-        await using Stream hashStream = new MemoryStream(payload);
-        BlobIdentifier id = useContentAddressedStorage ? BlobIdentifier.FromContentHash(await VerifyContentMatchesHash(hashStream, identifier)) : identifier;
-
-        BlobIdentifier objectStoreIdentifier = await PutObjectToStores(ns, payload, id);
-        await _blobIndex.AddBlobToIndex(ns, id);
-
-        return objectStoreIdentifier;
-    }
-
-    public async Task<Uri?> MaybePutObjectWithRedirect(NamespaceId ns, BlobIdentifier identifier)
-    {
-        bool allowRedirectUris = _namespacePolicyResolver.GetPoliciesForNs(ns).AllowRedirectUris;
-        if (!allowRedirectUris)
-        {
-            return null;
-        }
-
-        IServerTiming? serverTiming = _httpContextAccessor.HttpContext?.RequestServices.GetService<IServerTiming>();
-
-        // we only attempt to upload to the last store
-        // assuming that any other store will pull from it when they lack content once the upload has finished
-        IBlobStore store = _blobStores.Last();
-        {
-            string storeName = store.GetType().Name;
-            using TelemetrySpan scope = _tracer.StartActiveSpan("put_blob_with_redirect")
-                    .SetAttribute("operation.name", "put_blob_with_redirect")
-                    .SetAttribute("resource.name", identifier.ToString())
-                    .SetAttribute("store", storeName)
-            ;
-
-            using ServerTimingMetricScoped? serverTimingScope = serverTiming?.CreateServerTimingMetricScope($"blob.put-redirect.{storeName}", $"PUT(redirect) to store: '{storeName}'");
-
-            Uri? redirectUri = await store.PutObjectWithRedirect(ns, identifier);
-            if (redirectUri != null)
-            {
-                return redirectUri;
-            }
-        }
-
-        // no store found that supports redirect
-        return null;
-    }
-
-    public async Task<BlobIdentifier> PutObjectKnownHash(NamespaceId ns, IBufferedPayload content, BlobIdentifier identifier)
-    {
-        using TelemetrySpan scope = _tracer.StartActiveSpan("put_blob")
-            .SetAttribute("operation.name", "put_blob")
-            .SetAttribute("resource.name", identifier.ToString())
-            .SetAttribute("Content-Length", content.Length.ToString())
-            ;
-
-        BlobIdentifier objectStoreIdentifier = await PutObjectToStores(ns, content, identifier);
-        await _blobIndex.AddBlobToIndex(ns, identifier);
-
-        return objectStoreIdentifier;
-    }
-
-    private async Task<BlobIdentifier> PutObjectToStores(NamespaceId ns, IBufferedPayload bufferedPayload, BlobIdentifier identifier)
-    {
-        IServerTiming? serverTiming = _httpContextAccessor.HttpContext?.RequestServices.GetService<IServerTiming>();
-
-        foreach (IBlobStore store in _blobStores)
-        {
-            using TelemetrySpan scope = _tracer.StartActiveSpan("put_blob_to_store")
-                .SetAttribute("operation.name", "put_blob_to_store")
-                .SetAttribute("resource.name", identifier.ToString())
-                .SetAttribute("store", store.GetType().ToString())
-                ;
-            string storeName = store.GetType().Name;
-
-            using ServerTimingMetricScoped? serverTimingScope = serverTiming?.CreateServerTimingMetricScope($"blob.put.{storeName}", $"PUT to store: '{storeName}'");
-
-            await using Stream s = bufferedPayload.GetStream();
-            await store.PutObject(ns, s, identifier);
-        }
-        NamespacePolicy policy = _namespacePolicyResolver.GetPoliciesForNs(ns);
-        if (policy.PopulateFallbackNamespaceOnUpload && policy.FallbackNamespace.HasValue)
-        {
-            await PutObjectToStores(policy.FallbackNamespace.Value, bufferedPayload, identifier);
-            await _blobIndex.AddBlobToIndex(policy.FallbackNamespace.Value, identifier);
-        }
-        return identifier;
-    }
-
-    private async Task<BlobIdentifier> PutObjectToStores(NamespaceId ns, byte[] payload, BlobIdentifier identifier)
-    {
-        foreach (IBlobStore store in _blobStores)
-        {
-            await store.PutObject(ns, payload, identifier);
-        }
-
-        NamespacePolicy policy = _namespacePolicyResolver.GetPoliciesForNs(ns);
-        if (policy.PopulateFallbackNamespaceOnUpload && policy.FallbackNamespace.HasValue)
-        {
-            await PutObjectToStores(policy.FallbackNamespace.Value, payload, identifier);
-            await _blobIndex.AddBlobToIndex(policy.FallbackNamespace.Value, identifier);
-        }
-        return identifier;
-    }
-
-    public async Task<BlobContents> GetObject(NamespaceId ns, BlobIdentifier blob, List<string>? storageLayers = null, bool supportsRedirectUri = false)
-    {
-        try
-        {
-            return await GetObjectFromStores(ns, blob, storageLayers, supportsRedirectUri);
-        }
-        catch (BlobNotFoundException)
-        {
-            NamespacePolicy policy = _namespacePolicyResolver.GetPoliciesForNs(ns);
-
-            if (policy.FallbackNamespace != null)
-            {
-                ClaimsPrincipal? user = _httpContextAccessor.HttpContext?.User;
-                HttpRequest? request = _httpContextAccessor.HttpContext?.Request;
-
-                if (user == null || request == null)
-                {
-                    _logger.LogWarning("Unable to fallback to namespace {Namespace} due to not finding required http context values", policy.FallbackNamespace.Value);
-                    throw;
-                }
-
-                if (_requestHelper == null)
-                {
-                    _logger.LogWarning("Unable to fallback to namespace {Namespace} due to request helper missing", policy.FallbackNamespace.Value);
-                    throw;
-                }
-
-                ActionResult? result = await _requestHelper.HasAccessToNamespace(user, request, policy.FallbackNamespace.Value, new [] { AclAction.ReadObject });
-                if (result != null)
-                {
-                    _logger.LogInformation("Authorization error when attempting to fallback to namespace {FallbackNamespace}. This may be confusing for users that as they had access to original namespace {Namespace}", policy.FallbackNamespace.Value, ns);
-                    throw new AuthorizationException(result, "Failed to authenticate for fallback namespace");
-                }
-
-                try
-                {
-                    IServerTiming? serverTiming = _httpContextAccessor.HttpContext?.RequestServices.GetService<IServerTiming>();
-
-                    // read the content from the fallback namespace
-                    BlobContents fallbackContent = await GetObjectFromStores(policy.FallbackNamespace.Value, blob, storageLayers);
-                    
-                    // populate the primary namespace with the content
-                    using TelemetrySpan _ = _tracer.StartActiveSpan("HierarchicalStore.Populate").SetAttribute("operation.name", "HierarchicalStore.Populate");
-                    using ServerTimingMetricScoped? serverTimingScope = serverTiming?.CreateServerTimingMetricScope($"blob.populate", "Populating caches with blob contents");
-
-                    await using MemoryStream tempStream = new MemoryStream();
-                    await fallbackContent.Stream.CopyToAsync(tempStream);
-                    byte[] data = tempStream.ToArray();
-
-                    await PutObject(ns, data, blob);
-                    return await GetObject(ns, blob);
-                }
-                catch (BlobNotFoundException)
-                {
-                    // if we fail to find the blob in the fallback namespace we can carry on as we will rethrow the blob not found exception later (if the replication also fails)
-                }
-            }
-
-            if (ShouldFetchBlobOnDemand(ns))
-            {
-                try
-                {
-                    return await ReplicateObject(ns, blob);
-                }
-                catch (BlobNotFoundException)
-                {
-                    // if the blob is not found we can ignore it as we will just rethrow it later anyway
-                }
-            }
-
-            // if the primary namespace failed check to see if we should use a fallback policy which has replication enabled
-            if (policy.FallbackNamespace != null && ShouldFetchBlobOnDemand(policy.FallbackNamespace.Value))
-            {
-                try
-                {
-                    return await ReplicateObject(policy.FallbackNamespace.Value, blob);
-                }
-                catch (BlobNotFoundException)
-                {
-                    // if the blob is not found we can ignore it as we will just rethrow it later anyway
-                }
-            }
-
-            // we might have attempted to fetch the object and failed, or had no fallback options, either way the blob can not be found so we should rethrow
-            throw;
-        }
-    }
-
-    private async Task<BlobContents> GetObjectFromStores(NamespaceId ns, BlobIdentifier blob, List<string>? storageLayers = null, bool supportsRedirectUri = false)
-    {
-        bool seenBlobNotFound = false;
-        bool seenNamespaceNotFound = false;
-        int numStoreMisses = 0;
-        BlobContents? blobContents = null;
-        IServerTiming? serverTiming = _httpContextAccessor.HttpContext?.RequestServices.GetService<IServerTiming>();
-
-        foreach (IBlobStore store in _blobStores)
-        {
-            string blobStoreName = store.GetType().Name;
-
-            // check which storage layers to skip if we have a explicit list of storage layers to use
-            if (storageLayers != null && storageLayers.Count != 0)
-            {
-                bool found = false;
-                foreach (string storageLayer in storageLayers)
-                {
-                    if (string.Equals(storageLayer, blobStoreName, StringComparison.OrdinalIgnoreCase))
-                    {
-                        found = true;
-                        break;
-                    }
-                }
-
-                if (!found)
-                {
-                    continue;
-                }
-            }
-
-            using TelemetrySpan scope = _tracer.StartActiveSpan("HierarchicalStore.GetObject")
-                .SetAttribute("operation.name", "HierarchicalStore.GetObject")
-                .SetAttribute("resource.name",  blob.ToString())
-                .SetAttribute("BlobStore", store.GetType().ToString())
-                .SetAttribute("ObjectFound", false.ToString())
-                ;
-
-            string storeName = store.GetType().Name;
-            using ServerTimingMetricScoped? serverTimingScope = serverTiming?.CreateServerTimingMetricScope($"blob.get.{storeName}", $"Blob GET from: '{storeName}'");
-            _storeGetAttemptsCounter?.Add(1, new KeyValuePair<string, object?>("store", storeName));
-            try
-            {
-                blobContents = await store.GetObject(ns, blob, supportsRedirectUri: supportsRedirectUri);
-                scope.SetAttribute("ObjectFound", true.ToString());
-                _storeGetHitsCounter?.Add(1, new KeyValuePair<string, object?>("store", storeName));
-                break;
-            }
-            catch (BlobNotFoundException)
-            {
-                seenBlobNotFound = true;
-                numStoreMisses++;
-            }
-            catch (NamespaceNotFoundException)
-            {
-                seenNamespaceNotFound = true;
-            }
-        }
-
-        if (seenBlobNotFound && blobContents == null)
-        {
-            throw new BlobNotFoundException(ns, blob);
-        }
-
-        if (seenNamespaceNotFound && blobContents == null)
-        {
-            throw new NamespaceNotFoundException(ns);
-        }
-
-        // if we applied filters to the storage layers resulting in no blob found we consider it a miss
-        if (storageLayers != null && storageLayers.Count != 0 && blobContents == null)
-        {
-            throw new BlobNotFoundException(ns, blob);
-        }
-
-        if (blobContents == null)
-        {
-            // Should not happen but exists to safeguard against the null pointer
-            throw new Exception("blobContents is null");
-        }
-
-        if (numStoreMisses >= 1)
-        {
-            using TelemetrySpan _ = _tracer.StartActiveSpan("HierarchicalStore.Populate").SetAttribute("operation.name", "HierarchicalStore.Populate");
-            using ServerTimingMetricScoped? serverTimingScope = serverTiming?.CreateServerTimingMetricScope($"blob.populate", "Populating caches with blob contents");
-
-            await using MemoryStream tempStream = new MemoryStream();
-            await blobContents.Stream.CopyToAsync(tempStream);
-            byte[] data = tempStream.ToArray();
-            
-            // Don't populate the last store, as that is where we got the hit
-            for (int i = 0; i < numStoreMisses; i++)
-            {
-                IBlobStore blobStore = _blobStores[i];
-                using TelemetrySpan scope = _tracer.StartActiveSpan("HierarchicalStore.PopulateStore")
-                    .SetAttribute("operation.name", "HierarchicalStore.PopulateStore")
-                    .SetAttribute("BlobStore", blobStore.GetType().Name);
-                // Populate each store traversed that did not have the content found lower in the hierarchy
-                await blobStore.PutObject(ns, data, blob);
-            }
-
-#pragma warning disable CA2000 // Dispose objects before losing scope , ownership is transfered to caller
-            blobContents = new BlobContents(data);
-#pragma warning restore CA2000 // Dispose objects before losing scope
-        }
-        
-        return blobContents;
-    }
-
-    public async Task<Uri?> GetObjectWithRedirect(NamespaceId ns, BlobIdentifier blob, List<string>? storageLayers = null)
-    {
-        bool seenBlobNotFound = false;
-        bool seenNamespaceNotFound = false;
-        int numStoreMisses = 0;
-        IServerTiming? serverTiming = _httpContextAccessor.HttpContext?.RequestServices.GetService<IServerTiming>();
-
-        Uri? redirectUri = null;
-        foreach (IBlobStore store in _blobStores)
-        {
-            string blobStoreName = store.GetType().Name;
-
-            // check which storage layers to skip if we have a explicit list of storage layers to use
-            if (storageLayers != null && storageLayers.Count != 0)
-            {
-                bool found = false;
-                foreach (string storageLayer in storageLayers)
-                {
-                    if (string.Equals(storageLayer, blobStoreName, StringComparison.OrdinalIgnoreCase))
-                    {
-                        found = true;
-                        break;
-                    }
-                }
-
-                if (!found)
-                {
-                    continue;
-                }
-            }
-
-            using TelemetrySpan scope = _tracer.StartActiveSpan("HierarchicalStore.GetObjectRedirect")
-                .SetAttribute("operation.name", "HierarchicalStore.GetObject")
-                .SetAttribute("resource.name",  blob.ToString())
-                .SetAttribute("BlobStore", store.GetType().ToString())
-                .SetAttribute("ObjectFound", false.ToString())
-                ;
-
-            string storeName = store.GetType().Name;
-            using ServerTimingMetricScoped? serverTimingScope = serverTiming?.CreateServerTimingMetricScope($"blob.get-redirect.{storeName}", $"Blob GET from: '{storeName}'");
-
-            try
-            {
-                redirectUri = await store.GetObjectByRedirect(ns, blob);
-                scope.SetAttribute("ObjectFound", true.ToString());
-                break;
-            }
-            catch (BlobNotFoundException)
-            {
-                seenBlobNotFound = true;
-                numStoreMisses++;
-            }
-            catch (NamespaceNotFoundException)
-            {
-                seenNamespaceNotFound = true;
-            }
-        }
-
-        if (seenBlobNotFound && redirectUri == null)
-        {
-            throw new BlobNotFoundException(ns, blob);
-        }
-
-        if (seenNamespaceNotFound && redirectUri == null)
-        {
-            throw new NamespaceNotFoundException(ns);
-        }
-
-        // if we applied filters to the storage layers resulting in no blob found we consider it a miss
-        if (storageLayers != null && storageLayers.Count != 0 && redirectUri == null)
-        {
-            throw new BlobNotFoundException(ns, blob);
-        }
-
-        return redirectUri;
-    }
-
-    public async Task<BlobContents> ReplicateObject(NamespaceId ns, BlobIdentifier blob, bool force = false)
-    {
-        if (!force && !ShouldFetchBlobOnDemand(ns))
-        {
-            throw new NotSupportedException($"Replication is not allowed in namespace {ns}");
-        }
-        IServerTiming? serverTiming = _httpContextAccessor.HttpContext?.RequestServices.GetService<IServerTiming>();
-        using TelemetrySpan scope = _tracer.StartActiveSpan("HierarchicalStore.Replicate").SetAttribute("operation.name", "HierarchicalStore.Replicate");
-
-        using ServerTimingMetricScoped? serverTimingScope = serverTiming?.CreateServerTimingMetricScope("blob.replicate", "Replicating blob from remote instances");
-
-        List<string> regions = await _blobIndex.GetBlobRegions(ns, blob);
-
-        if (!regions.Any())
-        {
-            throw new BlobReplicationException(ns, blob, "Blob not found in any region");
-        }
-
-        _logger.LogInformation("On-demand replicating blob {Blob} in Namespace {Namespace}", blob, ns);
-        List<(int, string)> possiblePeers = new List<(int, string)>(_peerStatusService.GetPeersByLatency(regions));
-
-        bool replicated = false;
-        foreach ((int latency, string? region) in possiblePeers)
-        {
-            PeerStatus? peerStatus = _peerStatusService.GetPeerStatus(region);
-            if (peerStatus == null)
-            {
-                throw new Exception($"Failed to find peer {region}");
-            }
-
-            _logger.LogInformation("Attempting to replicate blob {Blob} in Namespace {Namespace} from {Region}", blob, ns, region);
-
-            PeerEndpoints peerEndpoint = peerStatus.Endpoints.First();
-            using HttpClient httpClient = _httpClientFactory.CreateClient();
-            string url = peerEndpoint.Url.ToString();
-            if (!url.EndsWith("/", StringComparison.InvariantCultureIgnoreCase))
-            {
-                url += "/";
-            }
-            using HttpRequestMessage blobRequest = await BuildHttpRequest(HttpMethod.Get, new Uri($"{url}api/v1/blobs/{ns}/{blob}"));
-            HttpResponseMessage blobResponse = await httpClient.SendAsync(blobRequest, HttpCompletionOption.ResponseHeadersRead);
-
-            if (blobResponse.StatusCode == HttpStatusCode.NotFound)
-            {
-                // try the next region
-                continue;
-            }
-
-            if (blobResponse.StatusCode != HttpStatusCode.OK)
-            {
-                throw new BlobReplicationException(ns, blob, $"Failed to replicate {blob} in {ns} from region {region} due to bad http status code {blobResponse.StatusCode}.");
-            }
-
-            await using Stream s = await blobResponse.Content.ReadAsStreamAsync();
-            using IBufferedPayload payload = await _bufferedPayloadFactory.CreateFilesystemBufferedPayload(s);
-            await PutObject(ns, payload, blob);
-            replicated = true;
-            break;
-        }
-
-        if (!replicated)
-        {
-            throw new BlobReplicationException(ns, blob, $"Failed to replicate {blob} in {ns} due to it not existing in any region");
-        }
-
-        return await GetObject(ns, blob);
-    }
-
-    public bool ShouldFetchBlobOnDemand(NamespaceId ns)
-    {
-        return _settings.CurrentValue.EnableOnDemandReplication && _namespacePolicyResolver.GetPoliciesForNs(ns).OnDemandReplication;
-    }
-
-    private async Task<HttpRequestMessage> BuildHttpRequest(HttpMethod httpMethod, Uri uri)
-    {
-        string? token = await _serviceCredentials.GetTokenAsync();
-        HttpRequestMessage request = new HttpRequestMessage(httpMethod, uri);
-        if (!string.IsNullOrEmpty(token))
-        {
-            request.Headers.Add("Authorization", $"{_serviceCredentials.GetAuthenticationScheme()} {token}");
-        }
-
-        return request;
-    }
-
-    public async Task<bool> Exists(NamespaceId ns, BlobIdentifier blob, List<string>? storageLayers = null)
-    {
-        bool exists = await ExistsInStores(ns, blob, storageLayers);
-        if (exists)
-        {
-            return exists;
-        }
-
-        NamespacePolicy policy = _namespacePolicyResolver.GetPoliciesForNs(ns);
-        if (policy.FallbackNamespace != null)
-        {
-            return await ExistsInStores(policy.FallbackNamespace.Value, blob, storageLayers);
-        }
-
-        return false;
-    }
-
-    private async Task<bool> ExistsInStores(NamespaceId ns, BlobIdentifier blob, List<string>? storageLayers = null)
-    {
-        bool useBlobIndex = _namespacePolicyResolver.GetPoliciesForNs(ns).UseBlobIndexForExists;
-        if (useBlobIndex)
-        {
-            using TelemetrySpan scope = _tracer.StartActiveSpan("HierarchicalStore.ObjectExists")
-                .SetAttribute("operation.name", "HierarchicalStore.ObjectExists")
-                .SetAttribute("resource.name", blob.ToString())
-                .SetAttribute("BlobStore", "BlobIndex")
-                ;
-            bool exists = await _blobIndex.BlobExistsInRegion(ns, blob);
-            if (exists)
-            {
-                scope.SetAttribute("ObjectFound", true.ToString());
-                return true;
-            }
-
-            scope.SetAttribute("ObjectFound", false.ToString());
-            return false;
-        }
-        else
-        {
-            foreach (IBlobStore store in _blobStores)
-            {
-                string blobStoreName = store.GetType().Name;
-                // check which storage layers to skip if we have a explicit list of storage layers to use
-                if (storageLayers != null && storageLayers.Count != 0)
-                {
-                    bool found = false;
-                    foreach (string storageLayer in storageLayers)
-                    {
-                        if (string.Equals(storageLayer, blobStoreName, StringComparison.OrdinalIgnoreCase))
-                        {
-                            found = true;
-                            break;
-                        }
-                    }
-
-                    if (!found)
-                    {
-                        continue;
-                    }
-                }
-
-                using TelemetrySpan scope = _tracer.StartActiveSpan("HierarchicalStore.ObjectExists")
-                    .SetAttribute("operation.name", "HierarchicalStore.ObjectExists")
-                    .SetAttribute("resource.name", blob.ToString())
-                    .SetAttribute("BlobStore", blobStoreName)
-                    ;
-                if (await store.Exists(ns, blob))
-                {
-                    scope.SetAttribute("ObjectFound", true.ToString());
-                    return true;
-                }
-                scope.SetAttribute("ObjectFound", false.ToString());
-            }
-
-            return false;
-        }
-    }
-
-    public async Task<bool> ExistsInRootStore(NamespaceId ns, BlobIdentifier blob)
-    {
-        IBlobStore store = _blobStores.Last();
-
-        using TelemetrySpan scope = _tracer.StartActiveSpan("HierarchicalStore.ObjectExistsInRoot")
-            .SetAttribute("operation.name", "HierarchicalStore.ObjectExistsInRoot")
-            .SetAttribute("resource.name", blob.ToString())
-            .SetAttribute("BlobStore", store.GetType().Name)
-            ;
-        if (await store.Exists(ns, blob))
-        {
-            scope.SetAttribute("ObjectFound", true.ToString());
-            return true;
-        }
-        scope.SetAttribute("ObjectFound", false.ToString());
-        return false;
-    }
-
-    public async Task DeleteObject(NamespaceId ns, BlobIdentifier blob)
-    {
-        bool blobNotFound = false;
-        bool deletedAtLeastOnce = false;
-
-        // remove the object from the tracking first, if this times out we do not want to end up with a inconsistent blob index
-        // if the blob store delete fails on the other hand we will still run a delete again during GC (as the blob is still orphaned at that point)
-        // this assumes that blob gc is based on scanning the root blob store
-        await _blobIndex.RemoveBlobFromRegion(ns, blob);
-
-        foreach (IBlobStore store in _blobStores)
-        {
-            try
-            {
-                using TelemetrySpan scope = _tracer.StartActiveSpan("HierarchicalStore.DeleteObject")
-                    .SetAttribute("operation.name", "HierarchicalStore.DeleteObject")
-                    .SetAttribute("resource.name", blob.ToString())
-                    .SetAttribute("BlobStore", store.GetType().Name)
-                    ;
-
-                await store.DeleteObject(ns, blob);
-                deletedAtLeastOnce = true;
-            }
-            catch (NamespaceNotFoundException)
-            {
-                // Ignore
-            }
-            catch (BlobNotFoundException)
-            {
-                blobNotFound = true;
-            }
-        }
-
-        if (deletedAtLeastOnce)
-        {
-            return;
-        }
-
-        if (blobNotFound)
-        {
-            throw new BlobNotFoundException(ns, blob);
-        }
-
-        throw new NamespaceNotFoundException(ns);
-    }
-
-    public async Task DeleteNamespace(NamespaceId ns)
-    {
-        bool deletedAtLeastOnce = false;
-        foreach (IBlobStore store in _blobStores)
-        {
-            using TelemetrySpan scope = _tracer.StartActiveSpan("HierarchicalStore.DeleteNamespace")
-                .SetAttribute("operation.name", "HierarchicalStore.DeleteNamespace")
-                .SetAttribute("resource.name", ns.ToString())
-                .SetAttribute("BlobStore", store.GetType().Name)
-                ;
-            try
-            {
-                await store.DeleteNamespace(ns);
-                deletedAtLeastOnce = true;
-            }
-            catch (NamespaceNotFoundException)
-            {
-                // Ignore
-            }
-        }
-
-        if (deletedAtLeastOnce)
-        {
-            return;
-        }
-
-        throw new NamespaceNotFoundException(ns);
-    }
-
-    public IAsyncEnumerable<(BlobIdentifier,DateTime)> ListObjects(NamespaceId ns)
-    {
-        // as this is a hierarchy of blob stores the last blob store should contain the superset of all stores
-        return _blobStores.Last().ListObjects(ns);
-    }
-
-    public async Task<BlobIdentifier[]> FilterOutKnownBlobs(NamespaceId ns, IEnumerable<BlobIdentifier> blobs)
-    {
-        List<(BlobIdentifier, Task<bool>)> existTasks = new();
-        foreach (BlobIdentifier blob in blobs)
-        {
-            existTasks.Add((blob, Exists(ns, blob)));
-        }
-
-        List<BlobIdentifier> missingBlobs = new();
-        foreach ((BlobIdentifier blob, Task<bool> existsTask) in existTasks)
-        {
-            bool exists = await existsTask;
-
-            if (!exists)
-            {
-                missingBlobs.Add(blob);
-            }
-        }
-
-        return missingBlobs.ToArray();
-    }
-
-    public async Task<BlobIdentifier[]> FilterOutKnownBlobs(NamespaceId ns, IAsyncEnumerable<BlobIdentifier> blobs)
-    {
-        ConcurrentBag<BlobIdentifier> missingBlobs = new ConcurrentBag<BlobIdentifier>();
-
-        try
-        {
-            await Parallel.ForEachAsync(blobs, async (identifier, ctx) =>
-            {
-                bool exists = await Exists(ns, identifier);
-
-                if (!exists)
-                {
-                    missingBlobs.Add(identifier);
-                }
-            });
-        }
-        catch (AggregateException e)
-        {
-            if (e.InnerException is PartialReferenceResolveException)
-            {
-                throw e.InnerException;
-            }
-
-            throw;
-        }
-
-        return missingBlobs.ToArray();
-    }
-
-    public async Task<BlobContents> GetObjects(NamespaceId ns, BlobIdentifier[] blobs)
-    {
-        using TelemetrySpan _ = _tracer.StartActiveSpan("blob.combine").SetAttribute("operation.name", "blob.combine");
-        Task<BlobContents>[] tasks = new Task<BlobContents>[blobs.Length];
-        for (int i = 0; i < blobs.Length; i++)
-        {
-            tasks[i] = GetObject(ns, blobs[i]);
-        }
-
-        MemoryStream ms = new MemoryStream();
-        foreach (Task<BlobContents> task in tasks)
-        {
-            BlobContents blob = await task;
-            await using Stream s = blob.Stream;
-            await s.CopyToAsync(ms);
-        }
-
-        ms.Seek(0, SeekOrigin.Begin);
-
-        return new BlobContents(ms, ms.Length);
-    }
-}
-
-public static class BlobServiceExtensions
-{
-    public static async Task<ContentId> PutCompressedObject(this IBlobService blobService, NamespaceId ns, IBufferedPayload payload, ContentId? id, IServiceProvider provider)
-    {
-        IContentIdStore contentIdStore = provider.GetService<IContentIdStore>()!;
-        CompressedBufferUtils compressedBufferUtils = provider.GetService<CompressedBufferUtils>()!;
-        Tracer tracer = provider.GetService<Tracer>()!;
-
-        // decompress the content and generate a identifier from it to verify the identifier we got
-        await using Stream decompressStream = payload.GetStream();
-        // TODO: we should add a overload for decompress content that can work on streams, otherwise we are still limited to 2GB compressed blobs
-        byte[] decompressedContent = compressedBufferUtils.DecompressContent(await decompressStream.ToByteArray());
-
-        await using MemoryStream decompressedStream = new MemoryStream(decompressedContent);
-        ContentId identifierDecompressedPayload;
-        if (id != null)
-        {
-            identifierDecompressedPayload = ContentId.FromContentHash(await blobService.VerifyContentMatchesHash(decompressedStream, id));
-        }
-        else
-        {
-            ContentHash blobHash;
-            {
-                using TelemetrySpan _ = tracer.StartActiveSpan("web.hash").SetAttribute("operation.name", "web.hash");
-                blobHash = await BlobIdentifier.FromStream(decompressedStream);
-            }
-
-            identifierDecompressedPayload = ContentId.FromContentHash(blobHash);
-        }
-
-        BlobIdentifier identifierCompressedPayload;
-        {
-            using TelemetrySpan _ = tracer.StartActiveSpan("web.hash").SetAttribute("operation.name", "web.hash");
-            await using Stream hashStream = payload.GetStream();
-            identifierCompressedPayload = await BlobIdentifier.FromStream(hashStream);
-        }
-
-        // commit the mapping from the decompressed hash to the compressed hash, we run this in parallel with the blob store submit
-        // TODO: let users specify weight of the blob compared to previously submitted content ids
-        int contentIdWeight = (int)payload.Length;
-        Task contentIdStoreTask = contentIdStore.Put(ns, identifierDecompressedPayload, identifierCompressedPayload, contentIdWeight);
-
-        // we still commit the compressed buffer to the object store using the hash of the compressed content
-        {
-            await blobService.PutObjectKnownHash(ns, payload, identifierCompressedPayload);
-        }
-
-        await contentIdStoreTask;
-
-        return identifierDecompressedPayload;
-    }
-
-    public static async Task<(BlobContents, string)> GetCompressedObject(this IBlobService blobService, NamespaceId ns, ContentId contentId, IServiceProvider provider, bool supportsRedirectUri = false)
-    {
-        IContentIdStore contentIdStore = provider.GetService<IContentIdStore>()!;
-        Tracer tracer = provider.GetService<Tracer>()!;
-
-        BlobIdentifier[]? chunks = await contentIdStore.Resolve(ns, contentId, mustBeContentId: false);
-        if (chunks == null || chunks.Length == 0)
-        {
-            throw new ContentIdResolveException(contentId);
-        }
-
-        // single chunk, we just return that chunk
-        if (chunks.Length == 1)
-        {
-            BlobIdentifier blobToReturn = chunks[0];
-            string mimeType = CustomMediaTypeNames.UnrealCompressedBuffer;
-            if (contentId.Equals(blobToReturn))
-            {
-                // this was actually the unmapped blob, meaning its not a compressed buffer
-                mimeType = MediaTypeNames.Application.Octet;
-            }
-
-            return (await blobService.GetObject(ns, blobToReturn, supportsRedirectUri: supportsRedirectUri), mimeType);
-        }
-
-        // chunked content, combine the chunks into a single stream
-        using TelemetrySpan _ = tracer.StartActiveSpan("blob.combine").SetAttribute("operation.name", "blob.combine");
-        Task<BlobContents>[] tasks = new Task<BlobContents>[chunks.Length];
-        for (int i = 0; i < chunks.Length; i++)
-        {
-            // even if it was requested to support redirect, since we need to combine the chunks using redirects is not possible
-            tasks[i] = blobService.GetObject(ns, chunks[i], supportsRedirectUri: false);
-        }
-
-        MemoryStream ms = new MemoryStream();
-        foreach (Task<BlobContents> task in tasks)
-        {
-            BlobContents blob = await task;
-            await using Stream s = blob.Stream;
-            await s.CopyToAsync(ms);
-        }
-
-        ms.Seek(0, SeekOrigin.Begin);
-
-        // chunking could not have happened for a non compressed buffer so assume it is compressed
-        return (new BlobContents(ms, ms.Length), CustomMediaTypeNames.UnrealCompressedBuffer);
-    }
-}
\ No newline at end of file
diff --git a/Engine/Source/Programs/UnrealCloudDDC/Jupiter/Implementation/IObjectService.cs b/Engine/Source/Programs/UnrealCloudDDC/Jupiter/Implementation/IObjectService.cs
new file mode 100644
--- /dev/null
+++ b/Engine/Source/Programs/UnrealCloudDDC/Jupiter/Implementation/IObjectService.cs
@@ -0,0 +1,25 @@
+// Copyright Epic Games, Inc. All Rights Reserved.
+
+using System.Collections.Generic;
+using System.Threading.Tasks;
+using EpicGames.Horde.Storage;
+using EpicGames.Serialization;
+
+namespace Jupiter.Implementation
+{
+    public interface IObjectService
+    {
+        Task<(ObjectRecord, BlobContents?)> Get(NamespaceId ns, BucketId bucket, IoHashKey key, string[] fields, bool doLastAccessTracking = true);
+        Task<(ContentId[], BlobIdentifier[])> Put(NamespaceId ns, BucketId bucket, IoHashKey key, BlobIdentifier blobHash, CbObject payload);
+        Task<(ContentId[], BlobIdentifier[])> Finalize(NamespaceId ns, BucketId bucket, IoHashKey key, BlobIdentifier blobHash);
+
+        IAsyncEnumerable<NamespaceId> GetNamespaces();
+
+        Task<bool> Delete(NamespaceId ns, BucketId bucket, IoHashKey key);
+        Task<long> DropNamespace(NamespaceId ns);
+        Task<long> DeleteBucket(NamespaceId ns, BucketId bucket);
+
+        Task<bool> Exists(NamespaceId ns, BucketId bucket, IoHashKey key);
+        Task<List<BlobIdentifier>> GetReferencedBlobs(NamespaceId ns, BucketId bucket, IoHashKey key);
+    }
+}
diff --git a/Engine/Source/Programs/UnrealCloudDDC/Jupiter/Implementation/ObjectService.cs b/Engine/Source/Programs/UnrealCloudDDC/Jupiter/Implementation/ObjectService.cs
--- a/Engine/Source/Programs/UnrealCloudDDC/Jupiter/Implementation/ObjectService.cs
+++ b/Engine/Source/Programs/UnrealCloudDDC/Jupiter/Implementation/ObjectService.cs
@@ -17,22 +17,6 @@ using OpenTelemetry.Trace;
 
 namespace Jupiter.Implementation
 {
-    public interface IObjectService
-    {
-        Task<(ObjectRecord, BlobContents?)> Get(NamespaceId ns, BucketId bucket, IoHashKey key, string[] fields, bool doLastAccessTracking = true);
-        Task<(ContentId[], BlobIdentifier[])> Put(NamespaceId ns, BucketId bucket, IoHashKey key, BlobIdentifier blobHash, CbObject payload);
-        Task<(ContentId[], BlobIdentifier[])> Finalize(NamespaceId ns, BucketId bucket, IoHashKey key, BlobIdentifier blobHash);
-
-        IAsyncEnumerable<NamespaceId> GetNamespaces();
-
-        Task<bool> Delete(NamespaceId ns, BucketId bucket, IoHashKey key);
-        Task<long> DropNamespace(NamespaceId ns);
-        Task<long> DeleteBucket(NamespaceId ns, BucketId bucket);
-
-        Task<bool> Exists(NamespaceId ns, BucketId bucket, IoHashKey key);
-        Task<List<BlobIdentifier>> GetReferencedBlobs(NamespaceId ns, BucketId bucket, IoHashKey key);
-    }
-
     public class ObjectService : IObjectService
     {
         private readonly IHttpContextAccessor _httpContextAccessor;
