UnrealCloudDDC: Move IBlobService extension methods into IBlobService.cs.

diff --git a/Engine/Source/Programs/UnrealCloudDDC/Jupiter/Implementation/Blob/BlobService.cs b/Engine/Source/Programs/UnrealCloudDDC/Jupiter/Implementation/Blob/BlobService.cs
--- a/Engine/Source/Programs/UnrealCloudDDC/Jupiter/Implementation/Blob/BlobService.cs
+++ b/Engine/Source/Programs/UnrealCloudDDC/Jupiter/Implementation/Blob/BlobService.cs
@@ -846,104 +846,3 @@ public class BlobService : IBlobService
         return new BlobContents(ms, ms.Length);
     }
 }
-
-public static class BlobServiceExtensions
-{
-    public static async Task<ContentId> PutCompressedObject(this IBlobService blobService, NamespaceId ns, IBufferedPayload payload, ContentId? id, IServiceProvider provider)
-    {
-        IContentIdStore contentIdStore = provider.GetService<IContentIdStore>()!;
-        CompressedBufferUtils compressedBufferUtils = provider.GetService<CompressedBufferUtils>()!;
-        Tracer tracer = provider.GetService<Tracer>()!;
-
-        // decompress the content and generate a identifier from it to verify the identifier we got
-        await using Stream decompressStream = payload.GetStream();
-        // TODO: we should add a overload for decompress content that can work on streams, otherwise we are still limited to 2GB compressed blobs
-        byte[] decompressedContent = compressedBufferUtils.DecompressContent(await decompressStream.ToByteArray());
-
-        await using MemoryStream decompressedStream = new MemoryStream(decompressedContent);
-        ContentId identifierDecompressedPayload;
-        if (id != null)
-        {
-            identifierDecompressedPayload = ContentId.FromContentHash(await blobService.VerifyContentMatchesHash(decompressedStream, id));
-        }
-        else
-        {
-            ContentHash blobHash;
-            {
-                using TelemetrySpan _ = tracer.StartActiveSpan("web.hash").SetAttribute("operation.name", "web.hash");
-                blobHash = await BlobIdentifier.FromStream(decompressedStream);
-            }
-
-            identifierDecompressedPayload = ContentId.FromContentHash(blobHash);
-        }
-
-        BlobIdentifier identifierCompressedPayload;
-        {
-            using TelemetrySpan _ = tracer.StartActiveSpan("web.hash").SetAttribute("operation.name", "web.hash");
-            await using Stream hashStream = payload.GetStream();
-            identifierCompressedPayload = await BlobIdentifier.FromStream(hashStream);
-        }
-
-        // commit the mapping from the decompressed hash to the compressed hash, we run this in parallel with the blob store submit
-        // TODO: let users specify weight of the blob compared to previously submitted content ids
-        int contentIdWeight = (int)payload.Length;
-        Task contentIdStoreTask = contentIdStore.Put(ns, identifierDecompressedPayload, identifierCompressedPayload, contentIdWeight);
-
-        // we still commit the compressed buffer to the object store using the hash of the compressed content
-        {
-            await blobService.PutObjectKnownHash(ns, payload, identifierCompressedPayload);
-        }
-
-        await contentIdStoreTask;
-
-        return identifierDecompressedPayload;
-    }
-
-    public static async Task<(BlobContents, string)> GetCompressedObject(this IBlobService blobService, NamespaceId ns, ContentId contentId, IServiceProvider provider, bool supportsRedirectUri = false)
-    {
-        IContentIdStore contentIdStore = provider.GetService<IContentIdStore>()!;
-        Tracer tracer = provider.GetService<Tracer>()!;
-
-        BlobIdentifier[]? chunks = await contentIdStore.Resolve(ns, contentId, mustBeContentId: false);
-        if (chunks == null || chunks.Length == 0)
-        {
-            throw new ContentIdResolveException(contentId);
-        }
-
-        // single chunk, we just return that chunk
-        if (chunks.Length == 1)
-        {
-            BlobIdentifier blobToReturn = chunks[0];
-            string mimeType = CustomMediaTypeNames.UnrealCompressedBuffer;
-            if (contentId.Equals(blobToReturn))
-            {
-                // this was actually the unmapped blob, meaning its not a compressed buffer
-                mimeType = MediaTypeNames.Application.Octet;
-            }
-
-            return (await blobService.GetObject(ns, blobToReturn, supportsRedirectUri: supportsRedirectUri), mimeType);
-        }
-
-        // chunked content, combine the chunks into a single stream
-        using TelemetrySpan _ = tracer.StartActiveSpan("blob.combine").SetAttribute("operation.name", "blob.combine");
-        Task<BlobContents>[] tasks = new Task<BlobContents>[chunks.Length];
-        for (int i = 0; i < chunks.Length; i++)
-        {
-            // even if it was requested to support redirect, since we need to combine the chunks using redirects is not possible
-            tasks[i] = blobService.GetObject(ns, chunks[i], supportsRedirectUri: false);
-        }
-
-        MemoryStream ms = new MemoryStream();
-        foreach (Task<BlobContents> task in tasks)
-        {
-            BlobContents blob = await task;
-            await using Stream s = blob.Stream;
-            await s.CopyToAsync(ms);
-        }
-
-        ms.Seek(0, SeekOrigin.Begin);
-
-        // chunking could not have happened for a non compressed buffer so assume it is compressed
-        return (new BlobContents(ms, ms.Length), CustomMediaTypeNames.UnrealCompressedBuffer);
-    }
-}
\ No newline at end of file
diff --git a/Engine/Source/Programs/UnrealCloudDDC/Jupiter/Implementation/Blob/IBlobService.cs b/Engine/Source/Programs/UnrealCloudDDC/Jupiter/Implementation/Blob/IBlobService.cs
--- a/Engine/Source/Programs/UnrealCloudDDC/Jupiter/Implementation/Blob/IBlobService.cs
+++ b/Engine/Source/Programs/UnrealCloudDDC/Jupiter/Implementation/Blob/IBlobService.cs
@@ -3,9 +3,14 @@
 using System;
 using System.Collections.Generic;
 using System.IO;
+using System.Net.Mime;
 using System.Threading.Tasks;
+using EpicGames.AspNet;
 using EpicGames.Horde.Storage;
 using Jupiter.Common.Implementation;
+using Jupiter.Utils;
+using Microsoft.Extensions.DependencyInjection;
+using OpenTelemetry.Trace;
 
 namespace Jupiter.Implementation;
 
@@ -47,3 +52,103 @@ public interface IBlobService
     bool ShouldFetchBlobOnDemand(NamespaceId ns);
 }
 
+public static class BlobServiceExtensions
+{
+    public static async Task<ContentId> PutCompressedObject(this IBlobService blobService, NamespaceId ns, IBufferedPayload payload, ContentId? id, IServiceProvider provider)
+    {
+        IContentIdStore contentIdStore = provider.GetService<IContentIdStore>()!;
+        CompressedBufferUtils compressedBufferUtils = provider.GetService<CompressedBufferUtils>()!;
+        Tracer tracer = provider.GetService<Tracer>()!;
+
+        // decompress the content and generate a identifier from it to verify the identifier we got
+        await using Stream decompressStream = payload.GetStream();
+        // TODO: we should add a overload for decompress content that can work on streams, otherwise we are still limited to 2GB compressed blobs
+        byte[] decompressedContent = compressedBufferUtils.DecompressContent(await decompressStream.ToByteArray());
+
+        await using MemoryStream decompressedStream = new MemoryStream(decompressedContent);
+        ContentId identifierDecompressedPayload;
+        if (id != null)
+        {
+            identifierDecompressedPayload = ContentId.FromContentHash(await blobService.VerifyContentMatchesHash(decompressedStream, id));
+        }
+        else
+        {
+            ContentHash blobHash;
+            {
+                using TelemetrySpan _ = tracer.StartActiveSpan("web.hash").SetAttribute("operation.name", "web.hash");
+                blobHash = await BlobIdentifier.FromStream(decompressedStream);
+            }
+
+            identifierDecompressedPayload = ContentId.FromContentHash(blobHash);
+        }
+
+        BlobIdentifier identifierCompressedPayload;
+        {
+            using TelemetrySpan _ = tracer.StartActiveSpan("web.hash").SetAttribute("operation.name", "web.hash");
+            await using Stream hashStream = payload.GetStream();
+            identifierCompressedPayload = await BlobIdentifier.FromStream(hashStream);
+        }
+
+        // commit the mapping from the decompressed hash to the compressed hash, we run this in parallel with the blob store submit
+        // TODO: let users specify weight of the blob compared to previously submitted content ids
+        int contentIdWeight = (int)payload.Length;
+        Task contentIdStoreTask = contentIdStore.Put(ns, identifierDecompressedPayload, identifierCompressedPayload, contentIdWeight);
+
+        // we still commit the compressed buffer to the object store using the hash of the compressed content
+        {
+            await blobService.PutObjectKnownHash(ns, payload, identifierCompressedPayload);
+        }
+
+        await contentIdStoreTask;
+
+        return identifierDecompressedPayload;
+    }
+
+    public static async Task<(BlobContents, string)> GetCompressedObject(this IBlobService blobService, NamespaceId ns, ContentId contentId, IServiceProvider provider, bool supportsRedirectUri = false)
+    {
+        IContentIdStore contentIdStore = provider.GetService<IContentIdStore>()!;
+        Tracer tracer = provider.GetService<Tracer>()!;
+
+        BlobIdentifier[]? chunks = await contentIdStore.Resolve(ns, contentId, mustBeContentId: false);
+        if (chunks == null || chunks.Length == 0)
+        {
+            throw new ContentIdResolveException(contentId);
+        }
+
+        // single chunk, we just return that chunk
+        if (chunks.Length == 1)
+        {
+            BlobIdentifier blobToReturn = chunks[0];
+            string mimeType = CustomMediaTypeNames.UnrealCompressedBuffer;
+            if (contentId.Equals(blobToReturn))
+            {
+                // this was actually the unmapped blob, meaning its not a compressed buffer
+                mimeType = MediaTypeNames.Application.Octet;
+            }
+
+            return (await blobService.GetObject(ns, blobToReturn, supportsRedirectUri: supportsRedirectUri), mimeType);
+        }
+
+        // chunked content, combine the chunks into a single stream
+        using TelemetrySpan _ = tracer.StartActiveSpan("blob.combine").SetAttribute("operation.name", "blob.combine");
+        Task<BlobContents>[] tasks = new Task<BlobContents>[chunks.Length];
+        for (int i = 0; i < chunks.Length; i++)
+        {
+            // even if it was requested to support redirect, since we need to combine the chunks using redirects is not possible
+            tasks[i] = blobService.GetObject(ns, chunks[i], supportsRedirectUri: false);
+        }
+
+        MemoryStream ms = new MemoryStream();
+        foreach (Task<BlobContents> task in tasks)
+        {
+            BlobContents blob = await task;
+            await using Stream s = blob.Stream;
+            await s.CopyToAsync(ms);
+        }
+
+        ms.Seek(0, SeekOrigin.Begin);
+
+        // chunking could not have happened for a non compressed buffer so assume it is compressed
+        return (new BlobContents(ms, ms.Length), CustomMediaTypeNames.UnrealCompressedBuffer);
+    }
+}
