UnrealCloudDDC: Move ObjectHashMismatchException into IObjectService.cs.

diff --git a/Engine/Source/Programs/UnrealCloudDDC/Jupiter/Implementation/IObjectService.cs b/Engine/Source/Programs/UnrealCloudDDC/Jupiter/Implementation/IObjectService.cs
--- a/Engine/Source/Programs/UnrealCloudDDC/Jupiter/Implementation/IObjectService.cs
+++ b/Engine/Source/Programs/UnrealCloudDDC/Jupiter/Implementation/IObjectService.cs
@@ -1,5 +1,6 @@
 // Copyright Epic Games, Inc. All Rights Reserved.
 
+using System;
 using System.Collections.Generic;
 using System.Threading.Tasks;
 using EpicGames.Horde.Storage;
@@ -22,4 +23,11 @@ namespace Jupiter.Implementation
         Task<bool> Exists(NamespaceId ns, BucketId bucket, IoHashKey key);
         Task<List<BlobIdentifier>> GetReferencedBlobs(NamespaceId ns, BucketId bucket, IoHashKey key);
     }
+
+    public class ObjectHashMismatchException : Exception
+    {
+        public ObjectHashMismatchException(NamespaceId ns, BucketId bucket, IoHashKey name, BlobIdentifier suppliedHash, BlobIdentifier actualHash) : base($"Object {name} in bucket {bucket} and namespace {ns} did not reference hash {suppliedHash} was referencing {actualHash}")
+        {
+        }
+    }
 }
diff --git a/Engine/Source/Programs/UnrealCloudDDC/Jupiter/Implementation/Objects/ScyllaReferencesStore.cs b/Engine/Source/Programs/UnrealCloudDDC/Jupiter/Implementation/Objects/ScyllaReferencesStore.cs
--- a/Engine/Source/Programs/UnrealCloudDDC/Jupiter/Implementation/Objects/ScyllaReferencesStore.cs
+++ b/Engine/Source/Programs/UnrealCloudDDC/Jupiter/Implementation/Objects/ScyllaReferencesStore.cs
@@ -396,13 +396,6 @@ namespace Jupiter.Implementation
         }
     }
 
-    public class ObjectHashMismatchException : Exception
-    {
-        public ObjectHashMismatchException(NamespaceId ns, BucketId bucket, IoHashKey name, BlobIdentifier suppliedHash, BlobIdentifier actualHash) : base($"Object {name} in bucket {bucket} and namespace {ns} did not reference hash {suppliedHash} was referencing {actualHash}")
-        {
-        }
-    }
-
     public class ScyllaBlobIdentifier
     {
         public ScyllaBlobIdentifier()
