UnrealCloudDDC: Create an interface for IRequestHelper so that Horde can use a different implementation.

diff --git a/Engine/Source/Programs/UnrealCloudDDC/Jupiter/Controllers/AdminController.cs b/Engine/Source/Programs/UnrealCloudDDC/Jupiter/Controllers/AdminController.cs
--- a/Engine/Source/Programs/UnrealCloudDDC/Jupiter/Controllers/AdminController.cs
+++ b/Engine/Source/Programs/UnrealCloudDDC/Jupiter/Controllers/AdminController.cs
@@ -22,9 +22,9 @@ namespace Jupiter.Controllers
         private readonly LastAccessServiceReferences _lastAccessService;
         private readonly IRefCleanup _refCleanup;
         private readonly IConfiguration _configuration;
-        private readonly RequestHelper _requestHelper;
+        private readonly IRequestHelper _requestHelper;
 
-        public AdminController(LastAccessServiceReferences lastAccessService, IRefCleanup refCleanup, IConfiguration configuration, RequestHelper requestHelper)
+        public AdminController(LastAccessServiceReferences lastAccessService, IRefCleanup refCleanup, IConfiguration configuration, IRequestHelper requestHelper)
         {
             _lastAccessService = lastAccessService;
             _refCleanup = refCleanup;
diff --git a/Engine/Source/Programs/UnrealCloudDDC/Jupiter/Controllers/AuthController.cs b/Engine/Source/Programs/UnrealCloudDDC/Jupiter/Controllers/AuthController.cs
--- a/Engine/Source/Programs/UnrealCloudDDC/Jupiter/Controllers/AuthController.cs
+++ b/Engine/Source/Programs/UnrealCloudDDC/Jupiter/Controllers/AuthController.cs
@@ -16,11 +16,11 @@ namespace Jupiter.Controllers
     [Authorize]
     public class AuthController : ControllerBase
     {
-        private readonly RequestHelper _requestHelper;
+        private readonly IRequestHelper _requestHelper;
         private readonly INamespacePolicyResolver _namespacePolicyResolver;
         private readonly IOptionsMonitor<AuthSettings> _authSettings;
 
-        public AuthController(RequestHelper requestHelper, INamespacePolicyResolver namespacePolicyResolver, IOptionsMonitor<AuthSettings> authSettings)
+        public AuthController(IRequestHelper requestHelper, INamespacePolicyResolver namespacePolicyResolver, IOptionsMonitor<AuthSettings> authSettings)
         {
             _requestHelper = requestHelper;
             _namespacePolicyResolver = namespacePolicyResolver;
diff --git a/Engine/Source/Programs/UnrealCloudDDC/Jupiter/Controllers/BlobController.cs b/Engine/Source/Programs/UnrealCloudDDC/Jupiter/Controllers/BlobController.cs
--- a/Engine/Source/Programs/UnrealCloudDDC/Jupiter/Controllers/BlobController.cs
+++ b/Engine/Source/Programs/UnrealCloudDDC/Jupiter/Controllers/BlobController.cs
@@ -32,11 +32,11 @@ namespace Jupiter.Controllers
     {
         private readonly IBlobService _storage;
         private readonly IDiagnosticContext _diagnosticContext;
-        private readonly RequestHelper _requestHelper;
+        private readonly IRequestHelper _requestHelper;
         private readonly BufferedPayloadFactory _bufferedPayloadFactory;
         private readonly NginxRedirectHelper _nginxRedirectHelper;
 
-        public BlobsController(IBlobService storage, IDiagnosticContext diagnosticContext, RequestHelper requestHelper, BufferedPayloadFactory bufferedPayloadFactory, NginxRedirectHelper nginxRedirectHelper)
+        public BlobsController(IBlobService storage, IDiagnosticContext diagnosticContext, IRequestHelper requestHelper, BufferedPayloadFactory bufferedPayloadFactory, NginxRedirectHelper nginxRedirectHelper)
         {
             _storage = storage;
             _diagnosticContext = diagnosticContext;
diff --git a/Engine/Source/Programs/UnrealCloudDDC/Jupiter/Controllers/BundlesController.cs b/Engine/Source/Programs/UnrealCloudDDC/Jupiter/Controllers/BundlesController.cs
--- a/Engine/Source/Programs/UnrealCloudDDC/Jupiter/Controllers/BundlesController.cs
+++ b/Engine/Source/Programs/UnrealCloudDDC/Jupiter/Controllers/BundlesController.cs
@@ -168,7 +168,7 @@ namespace Jupiter.Controllers
     public class StorageController : ControllerBase
     {
         readonly IStorageService _storageService;
-        private readonly RequestHelper _requestHelper;
+        private readonly IRequestHelper _requestHelper;
         readonly IMemoryCache _memoryCache;
         readonly ILogger<StorageController> _logger;
 
@@ -179,7 +179,7 @@ namespace Jupiter.Controllers
         /// <param name="requestHelper"></param>
         /// <param name="memoryCache"></param>
         /// <param name="logger"></param>
-        public StorageController(IStorageService storageService, RequestHelper requestHelper, IMemoryCache memoryCache, ILogger<StorageController> logger)
+        public StorageController(IStorageService storageService, IRequestHelper requestHelper, IMemoryCache memoryCache, ILogger<StorageController> logger)
         {
             _storageService = storageService;
             _requestHelper = requestHelper;
diff --git a/Engine/Source/Programs/UnrealCloudDDC/Jupiter/Controllers/CompressedBlobController.cs b/Engine/Source/Programs/UnrealCloudDDC/Jupiter/Controllers/CompressedBlobController.cs
--- a/Engine/Source/Programs/UnrealCloudDDC/Jupiter/Controllers/CompressedBlobController.cs
+++ b/Engine/Source/Programs/UnrealCloudDDC/Jupiter/Controllers/CompressedBlobController.cs
@@ -30,11 +30,11 @@ namespace Jupiter.Controllers
         private readonly IBlobService _storage;
         private readonly IContentIdStore _contentIdStore;
         private readonly IDiagnosticContext _diagnosticContext;
-        private readonly RequestHelper _requestHelper;
+        private readonly IRequestHelper _requestHelper;
         private readonly BufferedPayloadFactory _bufferedPayloadFactory;
         private readonly NginxRedirectHelper _nginxRedirectHelper;
 
-        public CompressedBlobController(IBlobService storage, IContentIdStore contentIdStore, IDiagnosticContext diagnosticContext, RequestHelper requestHelper, BufferedPayloadFactory bufferedPayloadFactory, NginxRedirectHelper nginxRedirectHelper)
+        public CompressedBlobController(IBlobService storage, IContentIdStore contentIdStore, IDiagnosticContext diagnosticContext, IRequestHelper requestHelper, BufferedPayloadFactory bufferedPayloadFactory, NginxRedirectHelper nginxRedirectHelper)
         {
             _storage = storage;
             _contentIdStore = contentIdStore;
diff --git a/Engine/Source/Programs/UnrealCloudDDC/Jupiter/Controllers/ContentIdController.cs b/Engine/Source/Programs/UnrealCloudDDC/Jupiter/Controllers/ContentIdController.cs
--- a/Engine/Source/Programs/UnrealCloudDDC/Jupiter/Controllers/ContentIdController.cs
+++ b/Engine/Source/Programs/UnrealCloudDDC/Jupiter/Controllers/ContentIdController.cs
@@ -18,10 +18,10 @@ namespace Jupiter.Controllers
     [Authorize]
     public class ContentIdController : ControllerBase
     {
-        private readonly RequestHelper _requestHelper;
+        private readonly IRequestHelper _requestHelper;
         private readonly IContentIdStore _contentIdStore;
 
-        public ContentIdController(RequestHelper requestHelper, IContentIdStore contentIdStore)
+        public ContentIdController(IRequestHelper requestHelper, IContentIdStore contentIdStore)
         {
             _requestHelper = requestHelper;
             _contentIdStore = contentIdStore;
diff --git a/Engine/Source/Programs/UnrealCloudDDC/Jupiter/Controllers/IRequestHelper.cs b/Engine/Source/Programs/UnrealCloudDDC/Jupiter/Controllers/IRequestHelper.cs
new file mode 100644
--- /dev/null
+++ b/Engine/Source/Programs/UnrealCloudDDC/Jupiter/Controllers/IRequestHelper.cs
@@ -0,0 +1,16 @@
+// Copyright Epic Games, Inc. All Rights Reserved.
+
+using System.Security.Claims;
+using System.Threading.Tasks;
+using EpicGames.Horde.Storage;
+using Microsoft.AspNetCore.Http;
+using Microsoft.AspNetCore.Mvc;
+
+namespace Jupiter.Controllers
+{
+    public interface IRequestHelper
+    {
+        public Task<ActionResult?> HasAccessToNamespace(ClaimsPrincipal user, HttpRequest request, NamespaceId ns, JupiterAclAction[] aclActions);
+        public Task<ActionResult?> HasAccessForGlobalOperations(ClaimsPrincipal user, JupiterAclAction[] aclActions);
+    }
+}
diff --git a/Engine/Source/Programs/UnrealCloudDDC/Jupiter/Controllers/ObjectController.cs b/Engine/Source/Programs/UnrealCloudDDC/Jupiter/Controllers/ObjectController.cs
--- a/Engine/Source/Programs/UnrealCloudDDC/Jupiter/Controllers/ObjectController.cs
+++ b/Engine/Source/Programs/UnrealCloudDDC/Jupiter/Controllers/ObjectController.cs
@@ -31,13 +31,13 @@ namespace Jupiter.Controllers
     {
         private readonly IBlobService _storage;
         private readonly IDiagnosticContext _diagnosticContext;
-        private readonly RequestHelper _requestHelper;
+        private readonly IRequestHelper _requestHelper;
         private readonly IReferenceResolver _referenceResolver;
         private readonly BufferedPayloadFactory _bufferedPayloadFactory;
 
         private readonly ILogger _logger;
 
-        public ObjectController(IBlobService storage, IDiagnosticContext diagnosticContext, RequestHelper requestHelper, IReferenceResolver referenceResolver, BufferedPayloadFactory bufferedPayloadFactory, ILogger<ObjectController> logger)
+        public ObjectController(IBlobService storage, IDiagnosticContext diagnosticContext, IRequestHelper requestHelper, IReferenceResolver referenceResolver, BufferedPayloadFactory bufferedPayloadFactory, ILogger<ObjectController> logger)
         {
             _storage = storage;
             _diagnosticContext = diagnosticContext;
diff --git a/Engine/Source/Programs/UnrealCloudDDC/Jupiter/Controllers/ReferencesController.cs b/Engine/Source/Programs/UnrealCloudDDC/Jupiter/Controllers/ReferencesController.cs
--- a/Engine/Source/Programs/UnrealCloudDDC/Jupiter/Controllers/ReferencesController.cs
+++ b/Engine/Source/Programs/UnrealCloudDDC/Jupiter/Controllers/ReferencesController.cs
@@ -50,14 +50,14 @@ namespace Jupiter.Controllers
         private readonly BufferedPayloadFactory _bufferedPayloadFactory;
         private readonly IReferenceResolver _referenceResolver;
         private readonly NginxRedirectHelper _nginxRedirectHelper;
-        private readonly RequestHelper _requestHelper;
+        private readonly IRequestHelper _requestHelper;
         private readonly Tracer _tracer;
 
         private readonly ILogger _logger;
         private readonly IObjectService _objectService;
         private readonly IBlobService _blobStore;
 
-        public ReferencesController(IObjectService objectService, IBlobService blobStore, IDiagnosticContext diagnosticContext, FormatResolver formatResolver, BufferedPayloadFactory bufferedPayloadFactory, IReferenceResolver referenceResolver, NginxRedirectHelper nginxRedirectHelper, RequestHelper requestHelper, Tracer tracer, ILogger<ReferencesController> logger)
+        public ReferencesController(IObjectService objectService, IBlobService blobStore, IDiagnosticContext diagnosticContext, FormatResolver formatResolver, BufferedPayloadFactory bufferedPayloadFactory, IReferenceResolver referenceResolver, NginxRedirectHelper nginxRedirectHelper, IRequestHelper requestHelper, Tracer tracer, ILogger<ReferencesController> logger)
         {
             _objectService = objectService;
             _blobStore = blobStore;
diff --git a/Engine/Source/Programs/UnrealCloudDDC/Jupiter/Controllers/ReplicationLogController.cs b/Engine/Source/Programs/UnrealCloudDDC/Jupiter/Controllers/ReplicationLogController.cs
--- a/Engine/Source/Programs/UnrealCloudDDC/Jupiter/Controllers/ReplicationLogController.cs
+++ b/Engine/Source/Programs/UnrealCloudDDC/Jupiter/Controllers/ReplicationLogController.cs
@@ -26,11 +26,11 @@ namespace Jupiter.Controllers
     public class ReplicationLogController : ControllerBase
     {
         private readonly IServiceProvider _provider;
-        private readonly RequestHelper _requestHelper;
+        private readonly IRequestHelper _requestHelper;
         private readonly IReplicationLog _replicationLog;
         private readonly IOptionsMonitor<SnapshotSettings> _snapshotSettings;
 
-        public ReplicationLogController(IServiceProvider provider, RequestHelper requestHelper, IReplicationLog replicationLog, IOptionsMonitor<SnapshotSettings> snapshotSettings)
+        public ReplicationLogController(IServiceProvider provider, IRequestHelper requestHelper, IReplicationLog replicationLog, IOptionsMonitor<SnapshotSettings> snapshotSettings)
         {
             _provider = provider;
             _requestHelper = requestHelper;
diff --git a/Engine/Source/Programs/UnrealCloudDDC/Jupiter/Controllers/RequestHelper.cs b/Engine/Source/Programs/UnrealCloudDDC/Jupiter/Controllers/RequestHelper.cs
--- a/Engine/Source/Programs/UnrealCloudDDC/Jupiter/Controllers/RequestHelper.cs
+++ b/Engine/Source/Programs/UnrealCloudDDC/Jupiter/Controllers/RequestHelper.cs
@@ -14,7 +14,7 @@ using OpenTelemetry.Trace;
 
 namespace Jupiter.Controllers;
 
-public class RequestHelper
+public class RequestHelper : IRequestHelper
 {
     private readonly IAuthorizationService _authorizationService;
     private readonly INamespacePolicyResolver _namespacePolicyResolver;
diff --git a/Engine/Source/Programs/UnrealCloudDDC/Jupiter/Implementation/Blob/BlobService.cs b/Engine/Source/Programs/UnrealCloudDDC/Jupiter/Implementation/Blob/BlobService.cs
--- a/Engine/Source/Programs/UnrealCloudDDC/Jupiter/Implementation/Blob/BlobService.cs
+++ b/Engine/Source/Programs/UnrealCloudDDC/Jupiter/Implementation/Blob/BlobService.cs
@@ -37,7 +37,7 @@ public class BlobService : IBlobService
     private readonly IServiceCredentials _serviceCredentials;
     private readonly INamespacePolicyResolver _namespacePolicyResolver;
     private readonly IHttpContextAccessor _httpContextAccessor;
-    private readonly RequestHelper? _requestHelper;
+    private readonly IRequestHelper? _requestHelper;
     private readonly Tracer _tracer;
     private readonly BufferedPayloadFactory _bufferedPayloadFactory;
     private readonly ILogger _logger;
@@ -50,7 +50,7 @@ public class BlobService : IBlobService
         set => _blobStores = value.ToList();
     }
 
-    public BlobService(IServiceProvider provider, IOptionsMonitor<UnrealCloudDDCSettings> settings, IBlobIndex blobIndex, IPeerStatusService peerStatusService, IHttpClientFactory httpClientFactory, IServiceCredentials serviceCredentials, INamespacePolicyResolver namespacePolicyResolver, IHttpContextAccessor httpContextAccessor, RequestHelper? requestHelper, Tracer tracer, BufferedPayloadFactory bufferedPayloadFactory, ILogger<BlobService> logger, Meter? meter)
+    public BlobService(IServiceProvider provider, IOptionsMonitor<UnrealCloudDDCSettings> settings, IBlobIndex blobIndex, IPeerStatusService peerStatusService, IHttpClientFactory httpClientFactory, IServiceCredentials serviceCredentials, INamespacePolicyResolver namespacePolicyResolver, IHttpContextAccessor httpContextAccessor, IRequestHelper? requestHelper, Tracer tracer, BufferedPayloadFactory bufferedPayloadFactory, ILogger<BlobService> logger, Meter? meter)
     {
         _blobStores = GetBlobStores(provider, settings).ToList();
         _settings = settings;
diff --git a/Engine/Source/Programs/UnrealCloudDDC/Jupiter/Jupiter.Startup.cs b/Engine/Source/Programs/UnrealCloudDDC/Jupiter/Jupiter.Startup.cs
--- a/Engine/Source/Programs/UnrealCloudDDC/Jupiter/Jupiter.Startup.cs
+++ b/Engine/Source/Programs/UnrealCloudDDC/Jupiter/Jupiter.Startup.cs
@@ -155,7 +155,7 @@ namespace Jupiter
 
             services.AddSingleton(serviceType: typeof(ILeaderElection), CreateLeaderElection);
 
-            services.AddTransient<RequestHelper>();
+            services.AddTransient<IRequestHelper, RequestHelper>();
 
             services.AddSingleton(Configuration);
 
