UnrealCloudDDC: Move NamespaceNotFoundException into its own file, so it can be merged into Horde.

diff --git a/Engine/Source/Programs/UnrealCloudDDC/Jupiter.Common/Implementation/NamespacePolicyResolver.cs b/Engine/Source/Programs/UnrealCloudDDC/Jupiter.Common/Implementation/NamespacePolicyResolver.cs
--- a/Engine/Source/Programs/UnrealCloudDDC/Jupiter.Common/Implementation/NamespacePolicyResolver.cs
+++ b/Engine/Source/Programs/UnrealCloudDDC/Jupiter.Common/Implementation/NamespacePolicyResolver.cs
@@ -116,19 +116,4 @@ namespace Jupiter.Common
             throw new NamespaceNotFoundException(ns, $"Unable to find a valid policy for namespace {ns}");
         }
     }
-
-    public class NamespaceNotFoundException : Exception
-    {
-        public NamespaceId Namespace { get; }
-
-        public NamespaceNotFoundException(NamespaceId @namespace) : base($"Could not find namespace {@namespace}")
-        {
-            Namespace = @namespace;
-        }
-
-        public NamespaceNotFoundException(NamespaceId @namespace, string message) : base(message)
-        {
-            Namespace = @namespace;
-        }
-    }
 }
diff --git a/Engine/Source/Programs/UnrealCloudDDC/Jupiter.Common/NamespaceNotFoundException.cs b/Engine/Source/Programs/UnrealCloudDDC/Jupiter.Common/NamespaceNotFoundException.cs
new file mode 100644
--- /dev/null
+++ b/Engine/Source/Programs/UnrealCloudDDC/Jupiter.Common/NamespaceNotFoundException.cs
@@ -0,0 +1,22 @@
+// Copyright Epic Games, Inc. All Rights Reserved.
+
+using System;
+using EpicGames.Horde.Storage;
+
+namespace Jupiter.Common
+{
+    public class NamespaceNotFoundException : Exception
+    {
+        public NamespaceId Namespace { get; }
+
+        public NamespaceNotFoundException(NamespaceId @namespace) : base($"Could not find namespace {@namespace}")
+        {
+            Namespace = @namespace;
+        }
+
+        public NamespaceNotFoundException(NamespaceId @namespace, string message) : base(message)
+        {
+            Namespace = @namespace;
+        }
+    }
+}
