UnrealCloudDDC: Move AuthorizationException into IRequestHelper, so it's available to Horde code.

diff --git a/Engine/Source/Programs/UnrealCloudDDC/Jupiter/Controllers/IRequestHelper.cs b/Engine/Source/Programs/UnrealCloudDDC/Jupiter/Controllers/IRequestHelper.cs
--- a/Engine/Source/Programs/UnrealCloudDDC/Jupiter/Controllers/IRequestHelper.cs
+++ b/Engine/Source/Programs/UnrealCloudDDC/Jupiter/Controllers/IRequestHelper.cs
@@ -13,4 +13,14 @@ namespace Jupiter.Controllers
         public Task<ActionResult?> HasAccessToNamespace(ClaimsPrincipal user, HttpRequest request, NamespaceId ns, JupiterAclAction[] aclActions);
         public Task<ActionResult?> HasAccessForGlobalOperations(ClaimsPrincipal user, JupiterAclAction[] aclActions);
     }
+
+    public class AuthorizationException : Exception
+    {
+        public ActionResult Result { get; }
+
+        public AuthorizationException(ActionResult result, string errorMessage) : base(errorMessage)
+        {
+            Result = result;
+        }
+    }
 }
diff --git a/Engine/Source/Programs/UnrealCloudDDC/Jupiter/Controllers/RequestHelper.cs b/Engine/Source/Programs/UnrealCloudDDC/Jupiter/Controllers/RequestHelper.cs
--- a/Engine/Source/Programs/UnrealCloudDDC/Jupiter/Controllers/RequestHelper.cs
+++ b/Engine/Source/Programs/UnrealCloudDDC/Jupiter/Controllers/RequestHelper.cs
@@ -108,13 +108,3 @@ public class RequestHelper : IRequestHelper
         return null;
     }
 }
-
-public class AuthorizationException : Exception
-{
-    public ActionResult Result { get; }
-
-    public AuthorizationException(ActionResult result, string errorMessage) : base(errorMessage)
-    {
-        Result = result;
-    }
-}
\ No newline at end of file
