Missing using directive.

diff --git a/Engine/Source/Programs/UnrealCloudDDC/Jupiter/Controllers/IRequestHelper.cs b/Engine/Source/Programs/UnrealCloudDDC/Jupiter/Controllers/IRequestHelper.cs
--- a/Engine/Source/Programs/UnrealCloudDDC/Jupiter/Controllers/IRequestHelper.cs
+++ b/Engine/Source/Programs/UnrealCloudDDC/Jupiter/Controllers/IRequestHelper.cs
@@ -1,5 +1,6 @@
 // Copyright Epic Games, Inc. All Rights Reserved.
 
+using System;
 using System.Security.Claims;
 using System.Threading.Tasks;
 using EpicGames.Horde.Storage;
