Unreal Cloud DDC - Added option to use S3 transfer utility for uploads (which will use multipart when the file to upload is large enough)

diff --git a/Engine/Source/Programs/UnrealCloudDDC/Helm/UnrealCloudDDC/templates/configmap.yaml b/Engine/Source/Programs/UnrealCloudDDC/Helm/UnrealCloudDDC/templates/configmap.yaml
--- a/Engine/Source/Programs/UnrealCloudDDC/Helm/UnrealCloudDDC/templates/configmap.yaml
+++ b/Engine/Source/Programs/UnrealCloudDDC/Helm/UnrealCloudDDC/templates/configmap.yaml
@@ -166,5 +166,9 @@ S3:
   {{- if $s3.StoragePoolBucketOverride }}
   StoragePoolBucketOverride: {{ $s3.StoragePoolBucketOverride | toYaml | nindent 4  }}
   {{- end }}
+
+  {{- if $s3.UseMultiPartUpload }}
+  UseMultiPartUpload: {{ $s3.UseMultiPartUpload }}
+  {{- end }}
 {{ end }}
 
diff --git a/Engine/Source/Programs/UnrealCloudDDC/Jupiter/Implementation/Blob/AmazonS3Store.cs b/Engine/Source/Programs/UnrealCloudDDC/Jupiter/Implementation/Blob/AmazonS3Store.cs
--- a/Engine/Source/Programs/UnrealCloudDDC/Jupiter/Implementation/Blob/AmazonS3Store.cs
+++ b/Engine/Source/Programs/UnrealCloudDDC/Jupiter/Implementation/Blob/AmazonS3Store.cs
@@ -15,6 +15,7 @@ using KeyNotFoundException = System.Collections.Generic.KeyNotFoundException;
 using System.Threading;
 using System.Runtime.CompilerServices;
 using System.Collections.Concurrent;
+using Amazon.S3.Transfer;
 using Microsoft.Extensions.Logging;
 using Microsoft.Extensions.DependencyInjection;
 using OpenTelemetry.Trace;
@@ -251,6 +252,13 @@ namespace Jupiter.Implementation
 					}
 				}
 			}
+
+			if (_settings.CurrentValue.UseMultiPartUpload)
+			{
+				await WriteMultipart(path, stream, cancellationToken);
+			}
+			else
+			{
 			PutObjectRequest request = new PutObjectRequest
 			{
 				BucketName = _bucketName,
@@ -278,6 +286,31 @@ namespace Jupiter.Implementation
 				throw;
 			}
 		}
+		}
+
+		private async Task WriteMultipart(string path, Stream stream, CancellationToken cancellationToken)
+		{
+			using TransferUtility utility = new TransferUtility(_amazonS3);
+			try
+			{
+				await utility.UploadAsync(stream, _bucketName, path, cancellationToken);
+			}
+			catch (AmazonS3Exception e)
+			{
+				// if the same object is added twice S3 will raise a error, as we are content addressed we can just accept whichever of the objects so we can ignore that error
+				if (e.StatusCode == HttpStatusCode.Conflict)
+				{
+					return;
+				}
+
+				if (e.StatusCode == HttpStatusCode.TooManyRequests)
+				{
+					throw new ResourceHasToManyRequestsException(e);
+				}
+
+				throw;
+			}
+		}
 
 		public async Task<BlobContents?> TryReadAsync(string path, LastAccessTrackingFlags flags = LastAccessTrackingFlags.DoTracking, CancellationToken cancellationToken = default)
 		{
diff --git a/Engine/Source/Programs/UnrealCloudDDC/Jupiter/Settings.cs b/Engine/Source/Programs/UnrealCloudDDC/Jupiter/Settings.cs
--- a/Engine/Source/Programs/UnrealCloudDDC/Jupiter/Settings.cs
+++ b/Engine/Source/Programs/UnrealCloudDDC/Jupiter/Settings.cs
@@ -180,6 +180,7 @@ namespace Jupiter
 		public Dictionary<string, string> StoragePoolBucketOverride { get; set; } = new Dictionary<string, string>();
 
 		public bool? UseArnRegion { get; set; } = null;
+		public bool UseMultiPartUpload { get; set; } = false;
 	}
 
 	public class GCSettings
