Unreal Cloud DDC - Forward on Jupiter specific headers when using X-Accel in nginx

diff --git a/Engine/Source/Programs/UnrealCloudDDC/Helm/UnrealCloudDDC/templates/configmap-nginx.yaml b/Engine/Source/Programs/UnrealCloudDDC/Helm/UnrealCloudDDC/templates/configmap-nginx.yaml
--- a/Engine/Source/Programs/UnrealCloudDDC/Helm/UnrealCloudDDC/templates/configmap-nginx.yaml
+++ b/Engine/Source/Programs/UnrealCloudDDC/Helm/UnrealCloudDDC/templates/configmap-nginx.yaml
@@ -57,6 +57,11 @@ data:
           internal;
           alias {{ printf "%s/Blobs" .Values.persistence.mountPath }};
           {{ if .Values.nginx.addXAccelDebugHeader }} add_header X-Jupiter-XAccel-Used true; {{ end }}
+
+          add_header X-Jupiter-IoHash $upstream_http_x_jupiter_iohash;
+          add_header X-Jupiter-LastAccess $upstream_http_x_jupiter_lastaccess;
+          add_header X-Jupiter-InlinePayloadHash $upstream_http_x_jupiter_inlinepayloadhash;
+          add_header Server-Timing $upstream_http_server_timing;
         }
 
         location / {
@@ -95,6 +100,11 @@ data:
           internal;
           alias {{ printf "%s/Blobs" .Values.persistence.mountPath }};
           {{ if .Values.nginx.addXAccelDebugHeader }} add_header X-Jupiter-XAccel-Used true; {{ end }}
+
+          add_header X-Jupiter-IoHash $upstream_http_x_jupiter_iohash;
+          add_header X-Jupiter-LastAccess $upstream_http_x_jupiter_lastaccess;
+          add_header X-Jupiter-InlinePayloadHash $upstream_http_x_jupiter_inlinepayloadhash;
+          add_header Server-Timing $upstream_http_server_timing;
         }
 
         location / {
@@ -133,6 +143,11 @@ data:
           internal;
           alias {{ printf "%s/Blobs" .Values.persistence.mountPath }};
           {{ if .Values.nginx.addXAccelDebugHeader }} add_header X-Jupiter-XAccel-Used true; {{ end }}
+
+          add_header X-Jupiter-IoHash $upstream_http_x_jupiter_iohash;
+          add_header X-Jupiter-LastAccess $upstream_http_x_jupiter_lastaccess;
+          add_header X-Jupiter-InlinePayloadHash $upstream_http_x_jupiter_inlinepayloadhash;
+          add_header Server-Timing $upstream_http_server_timing;
         }
 
         location / {
@@ -171,6 +186,11 @@ data:
             internal;
             alias {{ printf "%s/Blobs" .Values.persistence.mountPath }};
             {{ if .Values.nginx.addXAccelDebugHeader }}add_header X-Jupiter-XAccel-Used true; {{ end }}
+
+            add_header X-Jupiter-IoHash $upstream_http_x_jupiter_iohash;
+            add_header X-Jupiter-LastAccess $upstream_http_x_jupiter_lastaccess;
+            add_header X-Jupiter-InlinePayloadHash $upstream_http_x_jupiter_inlinepayloadhash;
+            add_header Server-Timing $upstream_http_server_timing;
         }
 
         location / {
