Unreal Cloud DDC - Also expose auth token without formatting after interactive login as a environment variable

diff --git a/Engine/Source/Programs/UnrealCloudDDC/updateAuthHeaderUser.ps1 b/Engine/Source/Programs/UnrealCloudDDC/updateAuthHeaderUser.ps1
--- a/Engine/Source/Programs/UnrealCloudDDC/updateAuthHeaderUser.ps1
+++ b/Engine/Source/Programs/UnrealCloudDDC/updateAuthHeaderUser.ps1
@@ -1,4 +1,5 @@
 $TempFile = New-TemporaryFile
 
 & "dotnet" "..\..\..\..\Engine\Binaries\DotNET\OidcToken\portable\OidcToken.dll" "--Service=EpicGames-Okta" "--OutFile" $TempFile.FullName 
-$env:curl_auth_header = Get-Content -Raw -Path $TempFile.FullName | ConvertFrom-Json | ForEach-Object { "Authorization: Bearer $($_.Token)" }
\ No newline at end of file
+$env:curl_auth_header = Get-Content -Raw -Path $TempFile.FullName | ConvertFrom-Json | ForEach-Object { "Authorization: Bearer $($_.Token)" }
+$env:auth_token = Get-Content -Raw -Path $TempFile.FullName | ConvertFrom-Json | ForEach-Object { "$($_.Token)" }
