Unreal Cloud DDC - Check blobs existence in remote regions when determining if it exists and on demand replication is enabled, should resolve issue with replicating compressed buffers.

diff --git a/Engine/Source/Programs/UnrealCloudDDC/Jupiter/Implementation/Blob/BlobService.cs b/Engine/Source/Programs/UnrealCloudDDC/Jupiter/Implementation/Blob/BlobService.cs
--- a/Engine/Source/Programs/UnrealCloudDDC/Jupiter/Implementation/Blob/BlobService.cs
+++ b/Engine/Source/Programs/UnrealCloudDDC/Jupiter/Implementation/Blob/BlobService.cs
@@ -610,6 +610,11 @@ public class BlobService : IBlobService
 			return await ExistsInStores(policy.FallbackNamespace.Value, blob, storageLayers);
 		}
 
+		if (ShouldFetchBlobOnDemand(ns))
+        {
+            return await ExistsInRemote(ns, blob);
+        }
+
 		return false;
 	}
 
@@ -674,7 +679,27 @@ public class BlobService : IBlobService
 		}
 	}
 
-	public async Task<bool> ExistsInRootStore(NamespaceId ns, BlobIdentifier blob)
+	public async Task<bool> ExistsInRemote(NamespaceId ns, BlobIdentifier blob)
+	{
+        IServerTiming? serverTiming = _httpContextAccessor.HttpContext?.RequestServices.GetService<IServerTiming>();
+        using TelemetrySpan scope = _tracer.StartActiveSpan("HierarchicalStore.ExistsRemote").SetAttribute("operation.name", "HierarchicalStore.ExistsRemote");
+
+        using ServerTimingMetricScoped? serverTimingScope = serverTiming?.CreateServerTimingMetricScope("blob.exists-remote", "Verify if blob exists in remotes");
+
+        IOptions<JupiterSettings> jupiterSettings = _httpContextAccessor.HttpContext?.RequestServices.GetService<IOptions<JupiterSettings>>()!;
+        List<string> regions = await _blobIndex.GetBlobRegions(ns, blob);
+
+        // we do not actually verify that the blob exists remotely as that would take a lot of time
+        // instead we simply check if there are any regions were the blob exists that is not our current region
+        if (regions.Any(region => !string.Equals(region, jupiterSettings.Value.CurrentSite, StringComparison.OrdinalIgnoreCase)))
+        {
+            return true;
+        }
+
+		return false;
+	}
+
+    public async Task<bool> ExistsInRootStore(NamespaceId ns, BlobIdentifier blob)
 	{
 		IBlobStore store = _blobStores.Last();
 
