Unreal Cloud DDC - Attempted to work around limitation with s3 multipart uploading not uploading concurrently for streams by using the underlying files if present.

diff --git a/Engine/Source/Programs/UnrealCloudDDC/Jupiter.Common/Implementation/StreamingBodyPayload.cs b/Engine/Source/Programs/UnrealCloudDDC/Jupiter.Common/Implementation/StreamingBodyPayload.cs
--- a/Engine/Source/Programs/UnrealCloudDDC/Jupiter.Common/Implementation/StreamingBodyPayload.cs
+++ b/Engine/Source/Programs/UnrealCloudDDC/Jupiter.Common/Implementation/StreamingBodyPayload.cs
@@ -103,6 +103,8 @@ namespace Jupiter.Common.Implementation
 		private readonly FileInfo _tempFile;
 		private long _length;
 
+		public FileInfo TempFile => _tempFile;
+
 		private FilesystemBufferedPayload()
 		{
 			_tempFile = new FileInfo(Path.GetTempFileName());
diff --git a/Engine/Source/Programs/UnrealCloudDDC/Jupiter/Implementation/Blob/AmazonS3Store.cs b/Engine/Source/Programs/UnrealCloudDDC/Jupiter/Implementation/Blob/AmazonS3Store.cs
--- a/Engine/Source/Programs/UnrealCloudDDC/Jupiter/Implementation/Blob/AmazonS3Store.cs
+++ b/Engine/Source/Programs/UnrealCloudDDC/Jupiter/Implementation/Blob/AmazonS3Store.cs
@@ -16,6 +16,7 @@ using System.Threading;
 using System.Runtime.CompilerServices;
 using System.Collections.Concurrent;
 using Amazon.S3.Transfer;
+using Jupiter.Common.Implementation;
 using Microsoft.Extensions.Logging;
 using Microsoft.Extensions.DependencyInjection;
 using OpenTelemetry.Trace;
@@ -259,56 +260,85 @@ namespace Jupiter.Implementation
 			}
 			else
 			{
-			PutObjectRequest request = new PutObjectRequest
-			{
-				BucketName = _bucketName,
-				Key = path,
-				InputStream = stream
-			};
-
-			try
-			{
-				await _amazonS3.PutObjectAsync(request, cancellationToken);
-			}
-			catch (AmazonS3Exception e)
-			{
-				// if the same object is added twice S3 will raise a error, as we are content addressed we can just accept whichever of the objects so we can ignore that error
-				if (e.StatusCode == HttpStatusCode.Conflict)
+				PutObjectRequest request = new PutObjectRequest
 				{
-					return;
-				}
+					BucketName = _bucketName,
+					Key = path,
+					InputStream = stream
+				};
 
-				if (e.StatusCode == HttpStatusCode.TooManyRequests)
+				try
 				{
-					throw new ResourceHasToManyRequestsException(e);
+					await _amazonS3.PutObjectAsync(request, cancellationToken);
 				}
+				catch (AmazonS3Exception e)
+				{
+					// if the same object is added twice S3 will raise a error, as we are content addressed we can just accept whichever of the objects so we can ignore that error
+					if (e.StatusCode == HttpStatusCode.Conflict)
+					{
+						return;
+					}
 
-				throw;
+					if (e.StatusCode == HttpStatusCode.TooManyRequests)
+					{
+						throw new ResourceHasToManyRequestsException(e);
+					}
+
+					throw;
+				}
 			}
 		}
-		}
 
 		private async Task WriteMultipart(string path, Stream stream, CancellationToken cancellationToken)
 		{
-			using TransferUtility utility = new TransferUtility(_amazonS3);
+			FilesystemBufferedPayload? payload = null;
 			try
 			{
-				await utility.UploadAsync(stream, _bucketName, path, cancellationToken);
-			}
-			catch (AmazonS3Exception e)
-			{
-				// if the same object is added twice S3 will raise a error, as we are content addressed we can just accept whichever of the objects so we can ignore that error
-				if (e.StatusCode == HttpStatusCode.Conflict)
+				string? filePath = null;
+				if (stream is FileStream fileStream)
 				{
-					return;
+					filePath = fileStream.Name;
 				}
+				else if (stream.Length > 16 * (long)Math.Pow(2, 20))
+				{
+					// will be chunked by TransferUtility
+					using FilesystemBufferedPayloadWriter writer = new FilesystemBufferedPayloadWriter();
+					await stream.CopyToAsync(writer.GetWritableStream(), cancellationToken);
+					payload = writer.Done();
 
-				if (e.StatusCode == HttpStatusCode.TooManyRequests)
+					filePath = payload.TempFile.FullName;
+				}
+				using TransferUtility utility = new TransferUtility(_amazonS3);
+				try
 				{
-					throw new ResourceHasToManyRequestsException(e);
+					if (filePath != null)
+					{
+						await utility.UploadAsync(filePath, _bucketName, path, cancellationToken);
+					}
+					else
+					{
+						await utility.UploadAsync(stream, _bucketName, path, cancellationToken);
+					}
 				}
+				catch (AmazonS3Exception e)
+				{
+					// if the same object is added twice S3 will raise a error, as we are content addressed we can just accept whichever of the objects so we can ignore that error
+					if (e.StatusCode == HttpStatusCode.Conflict)
+					{
+						return;
+					}
 
-				throw;
+					if (e.StatusCode == HttpStatusCode.TooManyRequests)
+					{
+						throw new ResourceHasToManyRequestsException(e);
+					}
+
+					throw;
+				}
+			}
+			finally
+			{
+				payload?.Dispose();
 			}
 		}
 
@@ -402,7 +432,7 @@ namespace Jupiter.Implementation
 		Uri? GetPresignedUrl(string path, HttpVerb verb)
 		{
 			using TelemetrySpan span = _tracer.StartActiveSpan("s3.BuildPresignedUrl")
-				.SetAttribute("Path", path)
+				.SetAttribute("Path", path) 
 			;
 
 			try
