Unreal Cloud DDC - Fixed issue with multipart upload

diff --git a/Engine/Source/Programs/UnrealCloudDDC/Jupiter/Implementation/Blob/AmazonS3Store.cs b/Engine/Source/Programs/UnrealCloudDDC/Jupiter/Implementation/Blob/AmazonS3Store.cs
--- a/Engine/Source/Programs/UnrealCloudDDC/Jupiter/Implementation/Blob/AmazonS3Store.cs
+++ b/Engine/Source/Programs/UnrealCloudDDC/Jupiter/Implementation/Blob/AmazonS3Store.cs
@@ -303,7 +303,10 @@ namespace Jupiter.Implementation
 				{
 					// will be chunked by TransferUtility
 					using FilesystemBufferedPayloadWriter writer = new FilesystemBufferedPayloadWriter();
-					await stream.CopyToAsync(writer.GetWritableStream(), cancellationToken);
+					{
+						await using Stream writableStream = writer.GetWritableStream();
+						await stream.CopyToAsync(writableStream, cancellationToken);
+					}
 					payload = writer.Done();
 
 					filePath = payload.TempFile.FullName;
