Unreal Cloud DDC - Fixed AddAliasAsync overload build error

diff --git a/Engine/Source/Programs/UnrealCloudDDC/Jupiter/Implementation/Bundles/StorageService.cs b/Engine/Source/Programs/UnrealCloudDDC/Jupiter/Implementation/Bundles/StorageService.cs
--- a/Engine/Source/Programs/UnrealCloudDDC/Jupiter/Implementation/Bundles/StorageService.cs
+++ b/Engine/Source/Programs/UnrealCloudDDC/Jupiter/Implementation/Bundles/StorageService.cs
@@ -92,7 +92,7 @@ public class StorageClient : BundleStorageClient
 		return locator;
 	}
 
-	public override async Task AddAliasAsync(Utf8String name, BundleNodeHandle handle, CancellationToken cancellationToken = default)
+	public override async Task AddAliasAsync(Utf8String name, BundleNodeHandle handle, int rank, CancellationToken cancellationToken = default)
 	{
 		// TODO: Implement aliases
 		await Task.CompletedTask;
