Unreal Cloud DDC - Removed sleep in build script as we do that in the tests now (wait for scylla to have started)

diff --git a/Engine/Source/Programs/UnrealCloudDDC/BuildUnrealCloudDDC.xml b/Engine/Source/Programs/UnrealCloudDDC/BuildUnrealCloudDDC.xml
--- a/Engine/Source/Programs/UnrealCloudDDC/BuildUnrealCloudDDC.xml
+++ b/Engine/Source/Programs/UnrealCloudDDC/BuildUnrealCloudDDC.xml
@@ -13,9 +13,6 @@
 		<Node Name="Test Jupiter">
 			<Docker-Compose-Up File="$(ComposeFile)"/>
 
-			<!-- Sleep a while to let scylla finish starting up before we execute the tests -->
-			<Spawn Exe="/bin/bash" Arguments="-c 'sleep 100'" />
-
 			<!-- Run the unit tests, and override the dotnet sdk used to the global one for the machine rather then what is in the sync, to allow us to use newer versions then the rest of the engine has migrated to -->
 			<DotNet Arguments="test Engine/Source/Programs/UnrealCloudDDC/UnrealCloudDDC.sln -l console;verbosity=normal --filter TestCategory!=SlowTests" DotNetPath="/usr/share/dotnet/dotnet"/>
 		</Node>
