Fix Jupiter compile error.

diff --git a/Engine/Source/Programs/UnrealCloudDDC/Jupiter/Implementation/Bundles/StorageService.cs b/Engine/Source/Programs/UnrealCloudDDC/Jupiter/Implementation/Bundles/StorageService.cs
--- a/Engine/Source/Programs/UnrealCloudDDC/Jupiter/Implementation/Bundles/StorageService.cs
+++ b/Engine/Source/Programs/UnrealCloudDDC/Jupiter/Implementation/Bundles/StorageService.cs
@@ -129,15 +129,15 @@ public class StorageClient : BundleStorageClient
 		return await Bundle.FromStreamAsync(blobContents.Stream, cancellationToken);
 	}
 
-	public override async Task<ReadOnlyMemory<byte>> ReadBundleRangeAsync(BundleLocator locator, int offset, int length, CancellationToken cancellationToken)
+	public override async Task<ReadOnlyMemory<byte>> ReadBundleRangeAsync(BundleLocator locator, int offset, int? length, CancellationToken cancellationToken)
 	{
 		Bundle bundle = await ReadBundleAsync(locator, cancellationToken);
 		ReadOnlySequence<byte> sequence = bundle.AsSequence();
 		sequence = sequence.Slice(offset);
 
-		if(sequence.Length > length)
+		if(length != null && sequence.Length > length.Value)
 		{
-			sequence = sequence.Slice(0, length);
+			sequence = sequence.Slice(0, length.Value);
 		}
 
 		return sequence.AsSingleSegment();
