Unreal Cloud DDC - Disabled buffering of of upload payloads in nginx and just forwarding this directly to cloud ddc, reduces resoponse times for large uploads a bit as we start processing the upload quicker.

diff --git a/Engine/Source/Programs/UnrealCloudDDC/Helm/UnrealCloudDDC/templates/configmap-nginx.yaml b/Engine/Source/Programs/UnrealCloudDDC/Helm/UnrealCloudDDC/templates/configmap-nginx.yaml
--- a/Engine/Source/Programs/UnrealCloudDDC/Helm/UnrealCloudDDC/templates/configmap-nginx.yaml
+++ b/Engine/Source/Programs/UnrealCloudDDC/Helm/UnrealCloudDDC/templates/configmap-nginx.yaml
@@ -38,6 +38,7 @@ data:
         server_name public_http1;
 
         proxy_buffering off;
+        proxy_request_buffering off;
         proxy_http_version 1.1;
         
         location /api {
@@ -81,6 +82,7 @@ data:
         server_name corp_http1;
         
         proxy_buffering off;
+        proxy_request_buffering off;
         proxy_http_version 1.1;
         
         location /api {
@@ -124,6 +126,7 @@ data:
         server_name public_http2;
 
         proxy_buffering off;
+        proxy_request_buffering off;
         proxy_http_version 1.1;
         
         location /api {
@@ -167,6 +170,7 @@ data:
         server_name corp_http2;
 
         proxy_buffering off;
+        proxy_request_buffering off;
         proxy_http_version 1.1;
         
         location /api {
