Unreal Cloud DDC - Fixed build issues in cloud ddc

diff --git a/Engine/Source/Programs/UnrealCloudDDC/Jupiter/Controllers/BundlesController.cs b/Engine/Source/Programs/UnrealCloudDDC/Jupiter/Controllers/BundlesController.cs
--- a/Engine/Source/Programs/UnrealCloudDDC/Jupiter/Controllers/BundlesController.cs
+++ b/Engine/Source/Programs/UnrealCloudDDC/Jupiter/Controllers/BundlesController.cs
@@ -270,8 +270,7 @@ namespace Jupiter.Controllers
 			}
 			else if (offset != null && length != null)
 			{
-				ReadOnlyMemory<byte> memory = await client.ReadBundleRangeAsync(locator, offset.Value, length.Value, cancellationToken);
-				stream = new ReadOnlyMemoryStream(memory);
+				stream = await client.OpenAsync(locator, offset.Value, length.Value, cancellationToken);
 			}
 			else
 			{
diff --git a/Engine/Source/Programs/UnrealCloudDDC/Jupiter/Implementation/Bundles/StorageService.cs b/Engine/Source/Programs/UnrealCloudDDC/Jupiter/Implementation/Bundles/StorageService.cs
--- a/Engine/Source/Programs/UnrealCloudDDC/Jupiter/Implementation/Bundles/StorageService.cs
+++ b/Engine/Source/Programs/UnrealCloudDDC/Jupiter/Implementation/Bundles/StorageService.cs
@@ -4,6 +4,7 @@ using System;
 using System.Buffers;
 using System.Collections.Concurrent;
 using System.Collections.Generic;
+using System.IO;
 using System.Runtime.CompilerServices;
 using System.Threading;
 using System.Threading.Tasks;
@@ -122,25 +123,11 @@ public class StorageClient : BundleStorageClient
 		return redirectUri;
 	}
 
-	public override async Task<Bundle> ReadBundleAsync(BundleLocator locator, CancellationToken cancellationToken)
+	public override async Task<Stream> OpenAsync(BundleLocator locator, int offset, int length = 0, CancellationToken cancellationToken = default)
 	{
 		BlobId blobIdentifier = BlobId.FromBlobLocator(locator);
 		BlobContents blobContents = await _blobService.GetObjectAsync(_namespaceId, blobIdentifier);
-		return await Bundle.FromStreamAsync(blobContents.Stream, cancellationToken);
-	}
-
-	public override async Task<ReadOnlyMemory<byte>> ReadBundleRangeAsync(BundleLocator locator, int offset, int? length, CancellationToken cancellationToken)
-	{
-		Bundle bundle = await ReadBundleAsync(locator, cancellationToken);
-		ReadOnlySequence<byte> sequence = bundle.AsSequence();
-		sequence = sequence.Slice(offset);
-
-		if(length != null && sequence.Length > length.Value)
-		{
-			sequence = sequence.Slice(0, length.Value);
-		}
-
-		return sequence.AsSingleSegment();
+		return blobContents.Stream;
 	}
 
 	public override async Task<BundleNodeHandle?> TryReadRefTargetAsync(RefName name, RefCacheTime cacheTime = default, CancellationToken cancellationToken = default)
