Unreal Cloud DDC - Introduce option to avoid schema changes - making it possible to manage the schema more explicitly.

diff --git a/Engine/Source/Programs/UnrealCloudDDC/Jupiter/Implementation/Objects/ScyllaSettings.cs b/Engine/Source/Programs/UnrealCloudDDC/Jupiter/Implementation/Objects/ScyllaSettings.cs
--- a/Engine/Source/Programs/UnrealCloudDDC/Jupiter/Implementation/Objects/ScyllaSettings.cs
+++ b/Engine/Source/Programs/UnrealCloudDDC/Jupiter/Implementation/Objects/ScyllaSettings.cs
@@ -101,6 +101,11 @@ public class ScyllaSettings : IValidatableObject
 	/// </summary>
 	public bool ListObjectsFromOldNamespaceTable { get; set; } = false;
 
+	/// <summary>
+	/// Set to disable any modifications to cassandra schema from Unreal Cloud DDC
+	/// </summary>
+	public bool AvoidSchemaChanges { get; set; } = false;
+
 	public IEnumerable<ValidationResult> Validate(ValidationContext validationContext)
 	{
 		List<ValidationResult> results = new List<ValidationResult>();
diff --git a/Engine/Source/Programs/UnrealCloudDDC/Jupiter/Jupiter.Startup.cs b/Engine/Source/Programs/UnrealCloudDDC/Jupiter/Jupiter.Startup.cs
--- a/Engine/Source/Programs/UnrealCloudDDC/Jupiter/Jupiter.Startup.cs
+++ b/Engine/Source/Programs/UnrealCloudDDC/Jupiter/Jupiter.Startup.cs
@@ -389,10 +389,13 @@ namespace Jupiter
 			}
 
 			string keyspace = replicatedSession.Keyspace;
-			replicatedSession.Execute(new SimpleStatement("CREATE TYPE IF NOT EXISTS blob_identifier (hash blob)"));
-			replicatedSession.UserDefinedTypes.Define(UdtMap.For<ScyllaBlobIdentifier>("blob_identifier", keyspace));
+			if (!settings.AvoidSchemaChanges)
+			{
+				replicatedSession.Execute(new SimpleStatement("CREATE TYPE IF NOT EXISTS blob_identifier (hash blob)"));
+				replicatedSession.Execute(new SimpleStatement("CREATE TYPE IF NOT EXISTS object_reference (bucket text, key text)"));
+			}
 
-			replicatedSession.Execute(new SimpleStatement("CREATE TYPE IF NOT EXISTS object_reference (bucket text, key text)"));
+			replicatedSession.UserDefinedTypes.Define(UdtMap.For<ScyllaBlobIdentifier>("blob_identifier", keyspace));
 			replicatedSession.UserDefinedTypes.Define(UdtMap.For<ScyllaObjectReference>("object_reference", keyspace));
 
 			string localKeyspaceName = $"{keyspace}_local_{settings.LocalKeyspaceSuffix}";
@@ -408,7 +411,10 @@ namespace Jupiter
 			replicatedSession.CreateKeyspaceIfNotExists(localKeyspaceName, replicationStrategyLocal);
 			ISession localSession = cluster.Connect(localKeyspaceName);
 
-			localSession.Execute(new SimpleStatement("CREATE TYPE IF NOT EXISTS blob_identifier (hash blob)"));
+			if (!settings.AvoidSchemaChanges)
+			{
+				localSession.Execute(new SimpleStatement("CREATE TYPE IF NOT EXISTS blob_identifier (hash blob)"));
+			}
 			localSession.UserDefinedTypes.Define(UdtMap.For<ScyllaBlobIdentifier>("blob_identifier", localKeyspaceName));
 
 			bool isScylla = !settings.UseAzureCosmosDB;
