Unreal Cloud DDC - Expanded avoid schema changes option to include all schema changes (table modifications was still being made)

diff --git a/Engine/Source/Programs/UnrealCloudDDC/Jupiter/Implementation/Blob/ScyllaBlobIndex.cs b/Engine/Source/Programs/UnrealCloudDDC/Jupiter/Implementation/Blob/ScyllaBlobIndex.cs
--- a/Engine/Source/Programs/UnrealCloudDDC/Jupiter/Implementation/Blob/ScyllaBlobIndex.cs
+++ b/Engine/Source/Programs/UnrealCloudDDC/Jupiter/Implementation/Blob/ScyllaBlobIndex.cs
@@ -30,33 +30,36 @@ public class ScyllaBlobIndex : IBlobIndex
 		_session = scyllaSessionManager.GetSessionForReplicatedKeyspace();
 		_mapper = new Mapper(_session);
 
-		string blobType = scyllaSessionManager.IsCassandra ? "blob" : "frozen<blob_identifier>";
-		_session.Execute(new SimpleStatement(@$"CREATE TABLE IF NOT EXISTS blob_index (
-			namespace text,
-			blob_id {blobType},
-			regions set<text>,
-			references set<frozen<object_reference>>,
-			PRIMARY KEY ((namespace, blob_id))
-		);"
-		));
-
-		_session.Execute(new SimpleStatement(@$"CREATE TABLE IF NOT EXISTS blob_index_v2 (
-			namespace text,
-			blob_id blob,
-			region text,
-			PRIMARY KEY ((namespace, blob_id), region)
-		);"
-		));
-
-		_session.Execute(new SimpleStatement(@$"CREATE TABLE IF NOT EXISTS blob_incoming_references (
-			namespace text,
-			blob_id blob,
-			reference_id blob,
-			reference_type smallint,
-			bucket_id text,
-			PRIMARY KEY ((namespace, blob_id), reference_id)
-		);"
-		));
+		if (!scyllaSettings.CurrentValue.AvoidSchemaChanges)
+		{
+			string blobType = scyllaSessionManager.IsCassandra ? "blob" : "frozen<blob_identifier>";
+			_session.Execute(new SimpleStatement(@$"CREATE TABLE IF NOT EXISTS blob_index (
+				namespace text,
+				blob_id {blobType},
+				regions set<text>,
+				references set<frozen<object_reference>>,
+				PRIMARY KEY ((namespace, blob_id))
+			);"
+			));
+
+			_session.Execute(new SimpleStatement(@$"CREATE TABLE IF NOT EXISTS blob_index_v2 (
+				namespace text,
+				blob_id blob,
+				region text,
+				PRIMARY KEY ((namespace, blob_id), region)
+			);"
+			));
+
+			_session.Execute(new SimpleStatement(@$"CREATE TABLE IF NOT EXISTS blob_incoming_references (
+				namespace text,
+				blob_id blob,
+				reference_id blob,
+				reference_type smallint,
+				bucket_id text,
+				PRIMARY KEY ((namespace, blob_id), reference_id)
+			);"
+			));
+		}
 	}
 
 	public async Task AddBlobToIndexAsync(NamespaceId ns, BlobId id, string? region = null)
diff --git a/Engine/Source/Programs/UnrealCloudDDC/Jupiter/Implementation/ContentId/ScyllaContentIdStore.cs b/Engine/Source/Programs/UnrealCloudDDC/Jupiter/Implementation/ContentId/ScyllaContentIdStore.cs
--- a/Engine/Source/Programs/UnrealCloudDDC/Jupiter/Implementation/ContentId/ScyllaContentIdStore.cs
+++ b/Engine/Source/Programs/UnrealCloudDDC/Jupiter/Implementation/ContentId/ScyllaContentIdStore.cs
@@ -6,6 +6,7 @@ using System.Threading.Tasks;
 using Cassandra;
 using Cassandra.Mapping;
 using EpicGames.Horde.Storage;
+using Microsoft.Extensions.Options;
 using OpenTelemetry.Trace;
 
 namespace Jupiter.Implementation
@@ -18,7 +19,7 @@ namespace Jupiter.Implementation
 		private readonly Mapper _mapper;
 		private readonly IScyllaSessionManager _scyllaSessionManager;
 
-		public ScyllaContentIdStore(IScyllaSessionManager scyllaSessionManager, IBlobService blobStore, Tracer tracer)
+		public ScyllaContentIdStore(IScyllaSessionManager scyllaSessionManager, IBlobService blobStore, Tracer tracer, IOptionsMonitor<ScyllaSettings> scyllaSettings)
 		{
 			_scyllaSessionManager = scyllaSessionManager;
 			_session = scyllaSessionManager.GetSessionForReplicatedKeyspace();
@@ -27,14 +28,17 @@ namespace Jupiter.Implementation
 
 			_mapper = new Mapper(_session);
 
-			string blobType = scyllaSessionManager.IsCassandra ? "blob" : "frozen<blob_identifier>";
-			_session.Execute(new SimpleStatement(@$"CREATE TABLE IF NOT EXISTS content_id (
-				content_id {blobType},
-				content_weight int, 
-				chunks set<{blobType}>, 
-				PRIMARY KEY ((content_id), content_weight)
-			);"
-			));
+			if (!scyllaSettings.CurrentValue.AvoidSchemaChanges)
+			{
+				string blobType = scyllaSessionManager.IsCassandra ? "blob" : "frozen<blob_identifier>";
+				_session.Execute(new SimpleStatement(@$"CREATE TABLE IF NOT EXISTS content_id (
+					content_id {blobType},
+					content_weight int, 
+					chunks set<{blobType}>, 
+					PRIMARY KEY ((content_id), content_weight)
+				);"
+				));
+			}
 		}
 
 		public async Task<BlobId[]?> ResolveAsync(NamespaceId ns, ContentId contentId, bool mustBeContentId)
diff --git a/Engine/Source/Programs/UnrealCloudDDC/Jupiter/Implementation/Objects/ScyllaReferencesStore.cs b/Engine/Source/Programs/UnrealCloudDDC/Jupiter/Implementation/Objects/ScyllaReferencesStore.cs
--- a/Engine/Source/Programs/UnrealCloudDDC/Jupiter/Implementation/Objects/ScyllaReferencesStore.cs
+++ b/Engine/Source/Programs/UnrealCloudDDC/Jupiter/Implementation/Objects/ScyllaReferencesStore.cs
@@ -41,40 +41,43 @@ namespace Jupiter.Implementation
 
 			_mapper = new Mapper(_session);
 
-			_session.Execute(new SimpleStatement(@"CREATE TABLE IF NOT EXISTS objects (
-				namespace text, 
-				bucket text, 
-				name text, 
-				payload_hash blob_identifier, 
-				inline_payload blob, 
-				is_finalized boolean,
-				last_access_time timestamp,
-				PRIMARY KEY ((namespace, bucket, name))
-			);"
-			));
-
-			_session.Execute(new SimpleStatement(@"CREATE TABLE IF NOT EXISTS object_last_access_v2 (
-				namespace text, 
-				bucket text, 
-				name text, 
-				last_access_time timestamp,
-				PRIMARY KEY ((namespace, bucket, name))
-			);"
-			));
-			
-			_session.Execute(new SimpleStatement(@"CREATE TABLE IF NOT EXISTS buckets_v2 (
-				namespace text, 
-				bucket text, 
-				PRIMARY KEY ((namespace), bucket)
-			);"
-			));
-
-			_session.Execute(new SimpleStatement(@"CREATE TABLE IF NOT EXISTS buckets (
-				namespace text, 
-				bucket set<text>, 
-				PRIMARY KEY (namespace)
-			);"
-			));
+			if (!_settings.CurrentValue.AvoidSchemaChanges)
+			{
+				_session.Execute(new SimpleStatement(@"CREATE TABLE IF NOT EXISTS objects (
+					namespace text, 
+					bucket text, 
+					name text, 
+					payload_hash blob_identifier, 
+					inline_payload blob, 
+					is_finalized boolean,
+					last_access_time timestamp,
+					PRIMARY KEY ((namespace, bucket, name))
+				);"
+				));
+
+				_session.Execute(new SimpleStatement(@"CREATE TABLE IF NOT EXISTS object_last_access_v2 (
+					namespace text, 
+					bucket text, 
+					name text, 
+					last_access_time timestamp,
+					PRIMARY KEY ((namespace, bucket, name))
+				);"
+				));
+				
+				_session.Execute(new SimpleStatement(@"CREATE TABLE IF NOT EXISTS buckets_v2 (
+					namespace text, 
+					bucket text, 
+					PRIMARY KEY ((namespace), bucket)
+				);"
+				));
+
+				_session.Execute(new SimpleStatement(@"CREATE TABLE IF NOT EXISTS buckets (
+					namespace text, 
+					bucket set<text>, 
+					PRIMARY KEY (namespace)
+				);"
+				));
+			}
 
 			// BYPASS CACHE is a scylla specific extension to disable populating the cache, should be ignored by other cassandra dbs
 			string cqlOptions = scyllaSessionManager.IsScylla ? "BYPASS CACHE" : "";
diff --git a/Engine/Source/Programs/UnrealCloudDDC/Jupiter/Implementation/TransactionLog/ScyllaReplicationLog.cs b/Engine/Source/Programs/UnrealCloudDDC/Jupiter/Implementation/TransactionLog/ScyllaReplicationLog.cs
--- a/Engine/Source/Programs/UnrealCloudDDC/Jupiter/Implementation/TransactionLog/ScyllaReplicationLog.cs
+++ b/Engine/Source/Programs/UnrealCloudDDC/Jupiter/Implementation/TransactionLog/ScyllaReplicationLog.cs
@@ -27,41 +27,44 @@ namespace Jupiter.Implementation
 			_session = scyllaSessionManager.GetSessionForLocalKeyspace();
 			_mapper = new Mapper(_session);
 
-			_session.Execute(new SimpleStatement(@"CREATE TABLE IF NOT EXISTS replication_log (
-				namespace varchar,
-				replication_bucket bigint,
-				replication_id timeuuid,
-				bucket varchar, 
-				key varchar, 
-				type int,
-				object_identifier blob_identifier,
-				PRIMARY KEY ((namespace, replication_bucket), replication_id)
-			);"
-			));
-
-			_session.Execute(new SimpleStatement(@"CREATE TABLE IF NOT EXISTS replication_snapshot (
-				namespace varchar,
-				id timeuuid,
-				blob_snapshot blob_identifier,
-				blob_namespace varchar,
-				PRIMARY KEY ((namespace), id)
-			) WITH CLUSTERING ORDER BY (id DESC);"  
-			));
-
-			_session.Execute(new SimpleStatement(@"CREATE TABLE IF NOT EXISTS replication_namespace (
-				namespace varchar,
-				PRIMARY KEY ((namespace))
-			);"  
-			));
-
-			_session.Execute(new SimpleStatement(@"CREATE TABLE IF NOT EXISTS replication_state (
-				namespace varchar,
-				name varchar,
-				last_bucket varchar,
-				last_event uuid,
-				PRIMARY KEY ((namespace), name)
-			);"  
-			));
+			if (!settings.CurrentValue.AvoidSchemaChanges)
+			{
+				_session.Execute(new SimpleStatement(@"CREATE TABLE IF NOT EXISTS replication_log (
+					namespace varchar,
+					replication_bucket bigint,
+					replication_id timeuuid,
+					bucket varchar, 
+					key varchar, 
+					type int,
+					object_identifier blob_identifier,
+					PRIMARY KEY ((namespace, replication_bucket), replication_id)
+				);"
+				));
+
+				_session.Execute(new SimpleStatement(@"CREATE TABLE IF NOT EXISTS replication_snapshot (
+					namespace varchar,
+					id timeuuid,
+					blob_snapshot blob_identifier,
+					blob_namespace varchar,
+					PRIMARY KEY ((namespace), id)
+				) WITH CLUSTERING ORDER BY (id DESC);"  
+				));
+
+				_session.Execute(new SimpleStatement(@"CREATE TABLE IF NOT EXISTS replication_namespace (
+					namespace varchar,
+					PRIMARY KEY ((namespace))
+				);"  
+				));
+
+				_session.Execute(new SimpleStatement(@"CREATE TABLE IF NOT EXISTS replication_state (
+					namespace varchar,
+					name varchar,
+					last_bucket varchar,
+					last_event uuid,
+					PRIMARY KEY ((namespace), name)
+				);"  
+				));
+			}
 		}
 
 		public async IAsyncEnumerable<NamespaceId> GetNamespacesAsync()
