Unreal Cloud DDC - Added metrics service which will calculate some bucket specific metrics every 24hrs, disabled by default as it puts quite a lot on the db.

Also includes the ability to list buckets and objects in buckets from the ref store, which are not very efficent and mostly meant to be used for background tasks like this.

diff --git a/Engine/Source/Programs/UnrealCloudDDC/Helm/UnrealCloudDDC/templates/configmap-worker.yaml b/Engine/Source/Programs/UnrealCloudDDC/Helm/UnrealCloudDDC/templates/configmap-worker.yaml
--- a/Engine/Source/Programs/UnrealCloudDDC/Helm/UnrealCloudDDC/templates/configmap-worker.yaml
+++ b/Engine/Source/Programs/UnrealCloudDDC/Helm/UnrealCloudDDC/templates/configmap-worker.yaml
@@ -86,6 +86,11 @@ Scylla:
 {{ $config.Scylla | toYaml | nindent 2 }}
 {{- end }}
 
+{{- if $config.Metrics }}
+Metrics: 
+{{ $config.Metrics | toYaml | nindent 2 }}
+{{- end }}
+
 {{- if $config.Serilog }}
 Serilog: 
 {{ $config.Serilog | toYaml | nindent 2 }}
diff --git a/Engine/Source/Programs/UnrealCloudDDC/Helm/UnrealCloudDDC/templates/configmap.yaml b/Engine/Source/Programs/UnrealCloudDDC/Helm/UnrealCloudDDC/templates/configmap.yaml
--- a/Engine/Source/Programs/UnrealCloudDDC/Helm/UnrealCloudDDC/templates/configmap.yaml
+++ b/Engine/Source/Programs/UnrealCloudDDC/Helm/UnrealCloudDDC/templates/configmap.yaml
@@ -125,6 +125,11 @@ Scylla:
 {{ $config.Scylla | toYaml | nindent 2 }}
 {{- end }}
 
+{{- if $config.Metrics }}
+Metrics: 
+{{ $config.Metrics | toYaml | nindent 2 }}
+{{- end }}
+
 {{- if $config.Serilog }}
 Serilog: 
 {{ $config.Serilog | toYaml | nindent 2 }}
diff --git a/Engine/Source/Programs/UnrealCloudDDC/Jupiter.Common/Utils/Gauge.cs b/Engine/Source/Programs/UnrealCloudDDC/Jupiter.Common/Utils/Gauge.cs
new file mode 100644
--- /dev/null
+++ b/Engine/Source/Programs/UnrealCloudDDC/Jupiter.Common/Utils/Gauge.cs
@@ -0,0 +1,49 @@
+// Copyright Epic Games, Inc. All Rights Reserved.
+
+using System.Collections.Generic;
+using System.Diagnostics.Metrics;
+using JetBrains.Annotations;
+
+namespace Jupiter.Common
+{
+	public class Gauge<T> where T : struct
+	{
+		[System.Diagnostics.CodeAnalysis.SuppressMessage("CodeQuality", "IDE0052:Remove unread private members", Justification = "Needed to keep the callback alive")] [UsedImplicitly]
+		private readonly ObservableGauge<T> _observableGauge;
+		private readonly List<Measurement<T>> _measurements = new List<Measurement<T>>();
+		private readonly object _lock = new object();
+
+		public Gauge(Meter meter, string name)
+		{
+			_observableGauge = meter.CreateObservableGauge(name, GetMeasurements);
+		}
+
+		private IEnumerable<Measurement<T>> GetMeasurements()
+		{
+			List<Measurement<T>> measurements;
+			lock (_lock)
+			{
+				measurements = new List<Measurement<T>>(_measurements);
+				_measurements.Clear();
+			}
+			return measurements;
+		}
+
+		public void Record(T measurement, IEnumerable<KeyValuePair<string, object?>> tags)
+		{
+			lock (_lock)
+			{
+				_measurements.Add(new Measurement<T>(measurement, tags));
+			}
+		}
+	}
+
+	public static class MeterExtensions
+	{
+		public static Gauge<T> CreateGauge<T>(this Meter m, string name) where T : struct
+		{
+			Gauge<T> gauge = new Gauge<T>(m, name);
+			return gauge;
+		}
+	}
+}
diff --git a/Engine/Source/Programs/UnrealCloudDDC/Jupiter/Implementation/MetricsService.cs b/Engine/Source/Programs/UnrealCloudDDC/Jupiter/Implementation/MetricsService.cs
new file mode 100644
--- /dev/null
+++ b/Engine/Source/Programs/UnrealCloudDDC/Jupiter/Implementation/MetricsService.cs
@@ -0,0 +1,224 @@
+// Copyright Epic Games, Inc. All Rights Reserved.
+
+using System;
+using System.Collections.Generic;
+using System.Diagnostics.Metrics;
+using System.Threading;
+using System.Threading.Tasks;
+using EpicGames.Horde.Storage;
+using EpicGames.Serialization;
+using Jupiter.Common;
+using Jupiter.Utils;
+using Microsoft.Extensions.DependencyInjection;
+using Microsoft.Extensions.Logging;
+using Microsoft.Extensions.Options;
+
+namespace Jupiter.Implementation
+{
+	public class MetricsState
+	{
+		public Task? CalculateMetricsTask { get; set; } = null;
+	}
+
+	public class MetricsServiceSettings
+	{
+		/// <summary>
+		/// Set to enable calulcation of metrics in the background. Adds load to database so only enable these if you intend to use it.
+		/// </summary>
+		public bool Enabled { get; set; } = false;
+
+		public TimeSpan PollFrequency { get; set; } = TimeSpan.FromHours(24);
+	}
+
+	public class MetricsService : PollingService<MetricsState>
+	{
+		private readonly IOptionsMonitor<MetricsServiceSettings> _settings;
+		private readonly IReferencesStore _referencesStore;
+		private volatile bool _alreadyPolling;
+
+		private readonly ILogger _logger;
+		private readonly IServiceProvider _provider;
+
+		public MetricsService(IOptionsMonitor<MetricsServiceSettings> settings, IReferencesStore referencesStore, ILogger<MetricsService> logger, IServiceProvider provider) : base(serviceName: nameof(MetricsService), settings.CurrentValue.PollFrequency, new MetricsState(), logger, startAtRandomTime: false)
+		{
+			_settings = settings;
+			_referencesStore = referencesStore;
+			_logger = logger;
+			_provider = provider;
+		}
+
+		protected override bool ShouldStartPolling()
+		{
+			return _settings.CurrentValue.Enabled;
+		}
+
+		public override async Task<bool> OnPollAsync(MetricsState state, CancellationToken cancellationToken)
+		{
+			if (_alreadyPolling)
+			{
+				return false;
+			}
+
+			_alreadyPolling = true;
+			try
+			{
+				if (!state.CalculateMetricsTask?.IsCompleted ?? false)
+				{
+					return false;
+				}
+
+				if (state.CalculateMetricsTask != null)
+				{
+					await state.CalculateMetricsTask;
+				}
+				state.CalculateMetricsTask = DoCalculateMetricsAsync(state, cancellationToken);
+
+				return true;
+
+			}
+			finally
+			{
+				_alreadyPolling = false;
+			}
+		}
+
+		private async Task DoCalculateMetricsAsync(MetricsState _, CancellationToken cancellationToken)
+		{
+			MetricsCalculator calculator = ActivatorUtilities.CreateInstance<MetricsCalculator>(_provider);
+			_logger.LogInformation("Attempting to calculate metrics. ");
+			try
+			{
+				await foreach (NamespaceId ns in _referencesStore.GetNamespacesAsync().WithCancellation(cancellationToken))
+				{
+					await foreach (BucketId bucket in _referencesStore.GetBuckets(ns).WithCancellation(cancellationToken))
+					{
+						DateTime start = DateTime.UtcNow;
+						
+						await calculator.CalculateStatsForBucketAsync(ns, bucket);
+
+						TimeSpan duration = DateTime.UtcNow - start;
+						_logger.LogInformation("Stats calculated for {Namespace} {Bucket} took {Duration}", ns, bucket, duration);
+					}
+				}
+			}
+			catch (Exception e)
+			{
+				_logger.LogError("Error calculating metrics. {Exception}",  e);
+			}
+		}
+	}
+
+	public class BucketStats
+	{
+		public NamespaceId Namespace { get; set; }
+		public BucketId Bucket { get; set; }
+		public long CountOfRefs { get; set; }
+		public long CountOfBlobs { get; set; }
+		public long TotalSize { get; set; }
+		public double AvgSize { get; set; }
+		public long LargestBlob { get; set; }
+		public long SmallestBlobFound { get; set; }
+	}
+
+	public class MetricsCalculator
+	{
+		private readonly IReferencesStore _referencesStore;
+		private readonly IBlobService _blobService;
+		private readonly IReferenceResolver _referenceResolver;
+
+		private readonly ILogger _logger;
+		private readonly Histogram<long> _blobSizeHistogram;
+		private readonly Gauge<double> _blobSizeAvgGauge;
+		private readonly Gauge<long> _blobSizeMinGauge;
+		private readonly Gauge<long> _blobSizeMaxGauge;
+		private readonly Gauge<long> _refsInBucketGauge;
+		private readonly Gauge<long> _blobSizeCountGauge;
+
+		public MetricsCalculator(IReferencesStore referencesStore, IBlobService blobService, IReferenceResolver referenceResolver, Meter meter, ILogger<MetricsService> logger)
+		{
+			_referencesStore = referencesStore;
+			_blobService = blobService;
+			_referenceResolver = referenceResolver;
+			_logger = logger;
+
+			_blobSizeHistogram = meter.CreateHistogram<long>("blobstats.size");
+			_blobSizeAvgGauge = meter.CreateGauge<double>("blobstats.bucket_size.avg");
+			_blobSizeMinGauge = meter.CreateGauge<long>("blobstats.bucket_size.min");
+			_blobSizeMaxGauge = meter.CreateGauge<long>("blobstats.bucket_size.max");
+			_blobSizeCountGauge = meter.CreateGauge<long>("blobstats.bucket_size.count");
+			_refsInBucketGauge = meter.CreateGauge<long>("blobstats.refs_in_bucket");
+		}
+
+		public async Task<BucketStats?> CalculateStatsForBucketAsync(NamespaceId ns, BucketId bucket)
+		{
+			KeyValuePair<string, object?>[] tags = new[] { new KeyValuePair<string, object?>("Bucket", bucket.ToString()) };
+
+			try
+			{
+				long countOfRefsInBucket = 0;
+				long sizeOfBlobsInBucket = 0;
+				long countOfBlobsInBucket = 0;
+				long largestBlobFound = 0;
+				long smallestBlobFound = long.MaxValue;
+				HashSet<byte[]> alreadyCountedBlobs = new HashSet<byte[]>(ByteArrayComparer.Default);
+				await foreach ((RefId _, BlobId blobId) in _referencesStore.GetRecordsInBucketAsync(ns, bucket))
+				{
+					countOfRefsInBucket += 1;
+
+					BlobContents blobContents = await _blobService.GetObjectAsync(ns, blobId);
+					byte[] rawBlob = await blobContents.Stream.ToByteArrayAsync();
+					CbObject cbObject = new CbObject(rawBlob);
+					// enumerate all referenced blobs from the ref
+					await foreach (BlobId blob in _referenceResolver.GetReferencedBlobs(ns, cbObject))
+					{
+						// check to see if we have counted this blob before
+						bool added = alreadyCountedBlobs.Add(blob.HashData);
+						if (added)
+						{
+							// new blob, lets count it
+							BlobContents referencedBlob = await _blobService.GetObjectAsync(ns, blob);
+							sizeOfBlobsInBucket += referencedBlob.Length;
+
+							smallestBlobFound = Math.Min(referencedBlob.Length, smallestBlobFound);
+							largestBlobFound = Math.Max(referencedBlob.Length, largestBlobFound);
+
+							countOfBlobsInBucket += 1;
+
+							_blobSizeHistogram.Record(referencedBlob.Length, tags);
+						}
+					}
+				}
+
+				double avgBlobSize = sizeOfBlobsInBucket / (double)countOfBlobsInBucket;
+				_blobSizeAvgGauge.Record(avgBlobSize, tags);
+				_blobSizeMinGauge.Record(smallestBlobFound, tags);
+				_blobSizeMaxGauge.Record(largestBlobFound, tags);
+				_blobSizeCountGauge.Record(countOfBlobsInBucket, tags);
+				_refsInBucketGauge.Record(countOfRefsInBucket, tags);
+
+				_logger.LogInformation("Stats calculated for {Namespace} {Bucket}. {CountOfRefs} {CountOfBlobs} {TotalSize} {AvgSize} {MaxSize} {MinSize}",
+					ns, bucket, countOfRefsInBucket, countOfBlobsInBucket, sizeOfBlobsInBucket, avgBlobSize, largestBlobFound, smallestBlobFound);
+
+				return new BucketStats
+				{
+					Namespace = ns,
+					Bucket = bucket,
+					CountOfRefs = countOfRefsInBucket,
+					CountOfBlobs = countOfBlobsInBucket,
+					TotalSize = sizeOfBlobsInBucket,
+					AvgSize = avgBlobSize,
+					LargestBlob = largestBlobFound,
+					SmallestBlobFound = smallestBlobFound
+				};
+
+			}
+			catch (Exception e)
+			{
+				_logger.LogError("Error calculating metrics for {Namespace} {Bucket} due to {Exception}", ns, bucket, e);
+			}
+
+			return null;
+		}
+	}
+
+}
diff --git a/Engine/Source/Programs/UnrealCloudDDC/Jupiter/Implementation/Objects/CacheReferencesStore.cs b/Engine/Source/Programs/UnrealCloudDDC/Jupiter/Implementation/Objects/CacheReferencesStore.cs
--- a/Engine/Source/Programs/UnrealCloudDDC/Jupiter/Implementation/Objects/CacheReferencesStore.cs
+++ b/Engine/Source/Programs/UnrealCloudDDC/Jupiter/Implementation/Objects/CacheReferencesStore.cs
@@ -77,11 +77,21 @@ namespace Jupiter.Implementation
 			throw new NotImplementedException("GetRecords not supported on a cached reference store");
 		}
 
+		public IAsyncEnumerable<(RefId, BlobId)> GetRecordsInBucketAsync(NamespaceId ns, BucketId bucket)
+		{
+			throw new NotImplementedException("GetRecordsInBucketAsync not supported on a cached reference store");
+		}
+
 		public IAsyncEnumerable<NamespaceId> GetNamespacesAsync()
 		{
 			throw new NotImplementedException("GetNamespaces not supported on a cached reference store");
 		}
 
+		public IAsyncEnumerable<BucketId> GetBuckets(NamespaceId ns)
+		{
+			throw new NotImplementedException("GetBuckets not supported on a cached reference store");
+		}
+
 		public Task<bool> DeleteAsync(NamespaceId ns, BucketId bucket, RefId key)
 		{
 			throw new NotImplementedException("Deletes are not supported on a cached reference store");
@@ -191,11 +201,21 @@ namespace Jupiter.Implementation
 			throw new NotImplementedException("GetRecords is not supported on a upstream reference store");
 		}
 
+		public IAsyncEnumerable<(RefId, BlobId)> GetRecordsInBucketAsync(NamespaceId ns, BucketId bucket)
+		{
+			throw new NotImplementedException("GetRecordsInBucketAsync is not supported on a upstream reference store");
+		}
+
 		public IAsyncEnumerable<NamespaceId> GetNamespacesAsync()
 		{
 			throw new NotImplementedException("GetNamespaces is not supported on a upstream reference store");
 		}
 
+		public IAsyncEnumerable<BucketId> GetBuckets(NamespaceId ns)
+		{
+			throw new NotImplementedException("GetBuckets is not supported on a upstream reference store");
+		}
+
 		public Task<bool> DeleteAsync(NamespaceId ns, BucketId bucket, RefId key)
 		{
 			throw new NotImplementedException("Delete is not supported on a upstream reference store");
diff --git a/Engine/Source/Programs/UnrealCloudDDC/Jupiter/Implementation/Objects/IReferencesStore.cs b/Engine/Source/Programs/UnrealCloudDDC/Jupiter/Implementation/Objects/IReferencesStore.cs
--- a/Engine/Source/Programs/UnrealCloudDDC/Jupiter/Implementation/Objects/IReferencesStore.cs
+++ b/Engine/Source/Programs/UnrealCloudDDC/Jupiter/Implementation/Objects/IReferencesStore.cs
@@ -32,7 +32,10 @@ namespace Jupiter.Implementation
 		Task UpdateLastAccessTimeAsync(NamespaceId ns, BucketId bucket, RefId key, DateTime newLastAccessTime);
 		IAsyncEnumerable<(NamespaceId, BucketId, RefId, DateTime)> GetRecordsAsync();
 
+		IAsyncEnumerable<(RefId, BlobId)> GetRecordsInBucketAsync(NamespaceId ns, BucketId bucket);
+
 		IAsyncEnumerable<NamespaceId> GetNamespacesAsync();
+		IAsyncEnumerable<BucketId> GetBuckets(NamespaceId ns);
 		Task<bool> DeleteAsync(NamespaceId ns, BucketId bucket, RefId key);
 		Task<long> DropNamespaceAsync(NamespaceId ns);
 		Task<long> DeleteBucketAsync(NamespaceId ns, BucketId bucket);
diff --git a/Engine/Source/Programs/UnrealCloudDDC/Jupiter/Implementation/Objects/MemoryCachedReferencesStore.cs b/Engine/Source/Programs/UnrealCloudDDC/Jupiter/Implementation/Objects/MemoryCachedReferencesStore.cs
--- a/Engine/Source/Programs/UnrealCloudDDC/Jupiter/Implementation/Objects/MemoryCachedReferencesStore.cs
+++ b/Engine/Source/Programs/UnrealCloudDDC/Jupiter/Implementation/Objects/MemoryCachedReferencesStore.cs
@@ -92,11 +92,21 @@ namespace Jupiter.Implementation.Objects
 			return _actualStore.GetRecordsAsync();
 		}
 
+		public IAsyncEnumerable<(RefId, BlobId)> GetRecordsInBucketAsync(NamespaceId ns, BucketId bucket)
+		{
+			return _actualStore.GetRecordsInBucketAsync(ns, bucket);
+		}
+
 		public IAsyncEnumerable<NamespaceId> GetNamespacesAsync()
 		{
 			return _actualStore.GetNamespacesAsync();
 		}
 
+		public IAsyncEnumerable<BucketId> GetBuckets(NamespaceId ns)
+		{
+			return _actualStore.GetBuckets(ns);
+		}
+
 		public Task<bool> DeleteAsync(NamespaceId ns, BucketId bucket, RefId key)
 		{
 			MemoryCache cache = GetCacheForNamespace(ns);
diff --git a/Engine/Source/Programs/UnrealCloudDDC/Jupiter/Implementation/Objects/MemoryReferencesStore.cs b/Engine/Source/Programs/UnrealCloudDDC/Jupiter/Implementation/Objects/MemoryReferencesStore.cs
--- a/Engine/Source/Programs/UnrealCloudDDC/Jupiter/Implementation/Objects/MemoryReferencesStore.cs
+++ b/Engine/Source/Programs/UnrealCloudDDC/Jupiter/Implementation/Objects/MemoryReferencesStore.cs
@@ -74,11 +74,35 @@ namespace Jupiter.Implementation
 			}
 		}
 
+		public async IAsyncEnumerable<(RefId, BlobId)> GetRecordsInBucketAsync(NamespaceId ns, BucketId bucket)
+		{
+			foreach (MemoryStoreObject o in _objects.Values.Where(o => o.Namespace == ns && o.Bucket == bucket))
+			{
+				await Task.CompletedTask;
+				yield return (o.Name, o.BlobHash);
+			}
+		}
+
 		public IAsyncEnumerable<NamespaceId> GetNamespacesAsync()
 		{
 			return _namespaces.ToAsyncEnumerable();
 		}
 
+		public async IAsyncEnumerable<BucketId> GetBuckets(NamespaceId ns)
+		{
+			HashSet<BucketId> buckets = new HashSet<BucketId>();
+			foreach (MemoryStoreObject o in _objects.Values.Where(o => o.Namespace == ns))
+			{
+				buckets.Add(o.Bucket);
+			}
+			await Task.CompletedTask;
+
+			foreach (BucketId bucket in buckets)
+			{
+				yield return bucket;
+			}
+		}
+
 		public Task<bool> DeleteAsync(NamespaceId ns, BucketId bucket, RefId key)
 		{
 			if (!_objects.TryRemove(BuildKey(ns, bucket, key), out MemoryStoreObject? _))
diff --git a/Engine/Source/Programs/UnrealCloudDDC/Jupiter/Implementation/Objects/MongoReferencesStore.cs b/Engine/Source/Programs/UnrealCloudDDC/Jupiter/Implementation/Objects/MongoReferencesStore.cs
--- a/Engine/Source/Programs/UnrealCloudDDC/Jupiter/Implementation/Objects/MongoReferencesStore.cs
+++ b/Engine/Source/Programs/UnrealCloudDDC/Jupiter/Implementation/Objects/MongoReferencesStore.cs
@@ -131,6 +131,20 @@ namespace Jupiter.Implementation
 			}
 		}
 
+		public async IAsyncEnumerable<(RefId, BlobId)> GetRecordsInBucketAsync(NamespaceId ns, BucketId bucket)
+		{
+			IMongoCollection<MongoReferencesModelV0> collection = GetCollection<MongoReferencesModelV0>();
+			IAsyncCursor<MongoReferencesModelV0>? cursor = await collection.FindAsync(m => m.Ns == ns.ToString() && m.Bucket == bucket.ToString());
+
+			while (await cursor.MoveNextAsync())
+			{
+				foreach (MongoReferencesModelV0 model in cursor.Current)
+				{
+					yield return (new RefId(model.Key), new BlobId(model.BlobIdentifier));
+				}
+			}
+		}
+
 		public async Task AddNamespaceIfNotExistAsync(NamespaceId ns)
 		{
 			FilterDefinition<MongoNamespacesModelV0> filter = Builders<MongoNamespacesModelV0>.Filter.Where(m => m.Ns == ns.ToString());
@@ -157,6 +171,27 @@ namespace Jupiter.Implementation
 			}
 		}
 
+		public async IAsyncEnumerable<BucketId> GetBuckets(NamespaceId ns)
+		{
+			IMongoCollection<MongoReferencesModelV0> collection = GetCollection<MongoReferencesModelV0>();
+
+			IAsyncCursor<MongoReferencesModelV0> cursor = await collection.FindAsync(m => m.Ns == ns.ToString());
+			
+			HashSet<BucketId> buckets = new HashSet<BucketId>();
+			while (await cursor.MoveNextAsync())
+			{
+				foreach (MongoReferencesModelV0? document in cursor.Current)
+				{
+					buckets.Add(new BucketId(document.Bucket));
+				}
+			}
+
+			foreach (BucketId bucket in buckets)
+			{
+				yield return bucket;
+			}
+		}
+
 		public async Task<bool> DeleteAsync(NamespaceId ns, BucketId bucket, RefId key)
 		{
 			IMongoCollection<MongoReferencesModelV0> collection = GetCollection<MongoReferencesModelV0>();
diff --git a/Engine/Source/Programs/UnrealCloudDDC/Jupiter/Implementation/Objects/ScyllaReferencesStore.cs b/Engine/Source/Programs/UnrealCloudDDC/Jupiter/Implementation/Objects/ScyllaReferencesStore.cs
--- a/Engine/Source/Programs/UnrealCloudDDC/Jupiter/Implementation/Objects/ScyllaReferencesStore.cs
+++ b/Engine/Source/Programs/UnrealCloudDDC/Jupiter/Implementation/Objects/ScyllaReferencesStore.cs
@@ -29,8 +29,10 @@ namespace Jupiter.Implementation
 		private readonly PreparedStatement _getNamespacesOldStatement;
 		private readonly PreparedStatement _getObjectsForPartitionRangeStatement;
 		private readonly PreparedStatement _getObjectsLastAccessForPartitionRangeStatement;
-		private readonly ConcurrentDictionary<NamespaceId, ConcurrentBag<BucketId>> _addedBuckets = new ConcurrentDictionary<NamespaceId, ConcurrentBag<BucketId>>();
+		private readonly PreparedStatement _getObjectsInBucketPartitionRangeStatement;
 
+		private readonly ConcurrentDictionary<NamespaceId, ConcurrentBag<BucketId>> _addedBuckets = new ConcurrentDictionary<NamespaceId, ConcurrentBag<BucketId>>();
+		
 		public ScyllaReferencesStore(IScyllaSessionManager scyllaSessionManager, IOptionsMonitor<ScyllaSettings> settings, INamespacePolicyResolver namespacePolicyResolver, Tracer tracer, ILogger<ScyllaReferencesStore> logger)
 		{
 			_session = scyllaSessionManager.GetSessionForReplicatedKeyspace();
@@ -88,6 +90,8 @@ namespace Jupiter.Implementation
 
 			_getObjectsForPartitionRangeStatement = _session.Prepare($"SELECT namespace, bucket, name, last_access_time FROM objects WHERE token(namespace, bucket, name) >= ? AND token(namespace, bucket, name) <= ? {cqlOptions}");
 			_getObjectsLastAccessForPartitionRangeStatement = _session.Prepare($"SELECT namespace, bucket, name, last_access_time FROM object_last_access_v2 WHERE token(namespace, bucket, name) >= ? AND token(namespace, bucket, name) <= ? {cqlOptions}");
+			
+			_getObjectsInBucketPartitionRangeStatement = _session.Prepare($"SELECT name, payload_hash FROM objects WHERE namespace = ? AND bucket = ? ALLOW FILTERING {cqlOptions}");
 		}
 
 		public async Task<RefRecord> GetAsync(NamespaceId ns, BucketId bucket, RefId name, IReferencesStore.FieldFlags fieldFlags, IReferencesStore.OperationFlags opFlags)
@@ -256,6 +260,27 @@ namespace Jupiter.Implementation
 			}
 		}
 
+		public async IAsyncEnumerable<(RefId, BlobId)> GetRecordsInBucketAsync(NamespaceId ns, BucketId bucket)
+		{
+			using TelemetrySpan scope = _tracer.BuildScyllaSpan("scylla.get_records_in_bucket_per_shard");
+			PreparedStatement getObjectStatement = _getObjectsInBucketPartitionRangeStatement;
+
+			RowSet rowSet = await _session.ExecuteAsync(getObjectStatement.Bind(ns.ToString(), bucket.ToString()));
+			foreach (Row row in rowSet)
+			{
+				string name = row.GetValue<string>("name");
+				ScyllaBlobIdentifier? blobIdentifier = row.GetValue<ScyllaBlobIdentifier>("payload_hash");
+
+				// skip any names that are not conformant to io hash
+				if (name.Length != 40)
+				{
+					continue;
+				}
+
+				yield return (new RefId(name), blobIdentifier.AsBlobIdentifier());
+			}
+		}
+
 		/// <summary>
 		/// This implements a more efficient scanning where we fetch objects based on which shard it is in. It scans the entire database and thus returns all namespaces.
 		/// See https://www.scylladb.com/2017/03/28/parallel-efficient-full-table-scan-scylla/
@@ -328,6 +353,19 @@ namespace Jupiter.Implementation
 			}
 		}
 
+		public async IAsyncEnumerable<BucketId> GetBuckets(NamespaceId ns)
+		{
+			foreach (ScyllaBucket? scyllaBucket in await _mapper.FetchAsync<ScyllaBucket>("WHERE ns = ?", ns))
+			{
+				if (scyllaBucket == null)
+				{
+					continue;
+				}
+
+				yield return new BucketId(scyllaBucket.Bucket);
+			}
+		}
+
 		public async Task<bool> DeleteAsync(NamespaceId ns, BucketId bucket, RefId key)
 		{
 			using TelemetrySpan scope = _tracer.BuildScyllaSpan("scylla.delete_record").SetAttribute("resource.name", $"{ns}.{bucket}.{key}");
diff --git a/Engine/Source/Programs/UnrealCloudDDC/Jupiter/Jupiter.Startup.cs b/Engine/Source/Programs/UnrealCloudDDC/Jupiter/Jupiter.Startup.cs
--- a/Engine/Source/Programs/UnrealCloudDDC/Jupiter/Jupiter.Startup.cs
+++ b/Engine/Source/Programs/UnrealCloudDDC/Jupiter/Jupiter.Startup.cs
@@ -106,6 +106,7 @@ namespace Jupiter
 			services.AddOptions<ReplicationSettings>().Configure(o => Configuration.GetSection("Replication").Bind(o)).ValidateDataAnnotations();
 			services.AddOptions<ServiceCredentialSettings>().Configure(o => Configuration.GetSection("ServiceCredentials").Bind(o)).ValidateDataAnnotations();
 			services.AddOptions<SnapshotSettings>().Configure(o => Configuration.GetSection("Snapshot").Bind(o)).ValidateDataAnnotations();
+			services.AddOptions<MetricsServiceSettings>().Configure(o => Configuration.GetSection("Metrics").Bind(o)).ValidateDataAnnotations();
 
 			services.AddOptions<ScyllaSettings>().Configure(o => Configuration.GetSection("Scylla").Bind(o)).ValidateDataAnnotations();
 
@@ -187,6 +188,9 @@ namespace Jupiter
 			services.AddSingleton<BlobIndexConsistencyCheckService>();
 			services.AddHostedService<BlobIndexConsistencyCheckService>(p => p.GetService<BlobIndexConsistencyCheckService>()!);
 
+			services.AddSingleton<MetricsService>();
+			services.AddHostedService<MetricsService>(p => p.GetService<MetricsService>()!);
+
 			services.AddSingleton(typeof(IPeerStatusService), typeof(PeerStatusService));
 			services.AddHostedService<PeerStatusService>(p => (PeerStatusService)p.GetService<IPeerStatusService>()!);
 		
