UnrealCloudDDC: Fix build errors.

diff --git a/Engine/Source/Programs/UnrealCloudDDC/Jupiter/Controllers/BundlesController.cs b/Engine/Source/Programs/UnrealCloudDDC/Jupiter/Controllers/BundlesController.cs
--- a/Engine/Source/Programs/UnrealCloudDDC/Jupiter/Controllers/BundlesController.cs
+++ b/Engine/Source/Programs/UnrealCloudDDC/Jupiter/Controllers/BundlesController.cs
@@ -13,6 +13,7 @@ using System.Threading.Tasks;
 using EpicGames.Core;
 using EpicGames.Horde.Storage;
 using EpicGames.Horde.Storage.Bundles;
+using EpicGames.Horde.Storage.Clients;
 using EpicGames.Horde.Storage.Nodes;
 using Jupiter.Implementation.Bundles;
 using Microsoft.AspNetCore.Authorization;
@@ -262,8 +263,8 @@ namespace Jupiter.Controllers
 			Stream stream;
 			if (offset == null && length == null)
 			{
-				Bundle bundle = await client.ReadBundleAsync(locator, cancellationToken);
-				stream = new ReadOnlySequenceStream(bundle.AsSequence());
+				byte[] bundle = await client.Backend.ReadBytesAsync(locator.ToString(), cancellationToken);
+				stream = new ReadOnlyMemoryStream(bundle);
 			}
 			else if (offset != null && length != null)
 			{
@@ -382,7 +383,7 @@ namespace Jupiter.Controllers
 
 			StorageClient storageClient = await _storageService.GetClientAsync(namespaceId, cancellationToken);
 
-			BundleReader reader = new BundleReader(storageClient, StorageCache.None, _logger);
+			BundleReader reader = new BundleReader(storageClient, BundleReaderCache.None, _logger);
 
 			BundleHeader header = await reader.ReadHeaderAsync(locator, cancellationToken);
 
@@ -457,7 +458,7 @@ namespace Jupiter.Controllers
 			}
 
 			StorageClient storageClient = await _storageService.GetClientAsync(namespaceId, cancellationToken);
-			BundleReader reader = new BundleReader(storageClient, StorageCache.None, _logger);
+			BundleReader reader = new BundleReader(storageClient, BundleReaderCache.None, _logger);
 
 			BundleHeader header = await reader.ReadHeaderAsync(locator, cancellationToken);
 			BundleExport export = header.Exports[exportIdx];
diff --git a/Engine/Source/Programs/UnrealCloudDDC/Jupiter/Implementation/Bundles/StorageService.cs b/Engine/Source/Programs/UnrealCloudDDC/Jupiter/Implementation/Bundles/StorageService.cs
--- a/Engine/Source/Programs/UnrealCloudDDC/Jupiter/Implementation/Bundles/StorageService.cs
+++ b/Engine/Source/Programs/UnrealCloudDDC/Jupiter/Implementation/Bundles/StorageService.cs
@@ -44,6 +44,13 @@ public class StorageService : IStorageService
 
 public sealed class JupiterStorageBackend : EpicGames.Horde.Storage.IStorageBackend
 {
+	class StorageObject : IStorageObject
+	{
+		public ReadOnlyMemory<byte> Data { get; }
+		public StorageObject(ReadOnlyMemory<byte> data) => Data = data;
+		public void Dispose() { }
+	}
+
 	private readonly NamespaceId _namespaceId;
 	private readonly IBlobService _blobService;
 	private readonly IBlobIndex _blobIndex;
@@ -67,7 +74,7 @@ public sealed class JupiterStorageBackend : EpicGames.Horde.Storage.IStorageBack
 
 	public Task<bool> ExistsAsync(string path, CancellationToken cancellationToken = default) => throw new NotSupportedException();
 
-	public async Task<Stream> ReadAsync(string path, int offset, int? length, CancellationToken cancellationToken = default)
+	public async Task<Stream> OpenAsync(string path, int offset, int? length, CancellationToken cancellationToken = default)
 	{
 		BlobId blobIdentifier = BlobId.FromBlobLocator(new BundleLocator(path));
 		BlobContents blobContents = await _blobService.GetObjectAsync(_namespaceId, blobIdentifier);
@@ -78,6 +85,15 @@ public sealed class JupiterStorageBackend : EpicGames.Horde.Storage.IStorageBack
 		return blobContents.Stream;
 	}
 
+	public async Task<IStorageObject> ReadAsync(string path, int offset, int? length, CancellationToken cancellationToken = default)
+	{
+		using (Stream stream = await OpenAsync(path, offset, length, cancellationToken))
+		{
+			byte[] data = await stream.ReadAllBytesAsync(cancellationToken);
+			return new StorageObject(data);
+		}
+	}
+
 	public async ValueTask<Uri?> TryGetReadRedirectAsync(string path, CancellationToken cancellationToken = default)
 	{
 		BlobId blobIdentifier = BlobId.FromBlobLocator(new BundleLocator(path));
@@ -128,7 +144,7 @@ public class StorageClient : BundleStorageClient
 	private readonly IBlobIndex _blobIndex;
 	private readonly BucketId _defaultBucket = new BucketId("bundles");
 
-	public JupiterStorageBackend Backend { get; }
+	public new JupiterStorageBackend Backend { get; }
 	public bool SupportsRedirects { get; set; } = true;
 
 	public StorageClient(NamespaceId namespaceId, IBlobService blobService, IReferencesStore refStore, IBlobIndex blobIndex)
@@ -139,7 +155,7 @@ public class StorageClient : BundleStorageClient
 	}
 
 	private StorageClient(JupiterStorageBackend backend, NamespaceId namespaceId, IReferencesStore refStore, IBlobIndex blobIndex)
-		: base(backend, StorageCache.None, NullLogger.Instance)
+		: base(backend, BundleReaderCache.None, NullLogger.Instance)
 	{
 		Backend = backend;
 
@@ -148,13 +164,13 @@ public class StorageClient : BundleStorageClient
 		_blobIndex = blobIndex;
 	}
 
-	public override async Task AddAliasAsync(Utf8String name, BundleNodeHandle handle, int rank, CancellationToken cancellationToken = default)
+	public override async Task AddAliasAsync(Utf8String name, BundleNodeLocator locator, int rank, CancellationToken cancellationToken = default)
 	{
 		// TODO: Implement aliases
 		await Task.CompletedTask;
 	}
 
-	public override async Task RemoveAliasAsync(Utf8String name, BundleNodeHandle handle, CancellationToken cancellationToken = default)
+	public override async Task RemoveAliasAsync(Utf8String name, BundleNodeLocator locator, CancellationToken cancellationToken = default)
 	{
 		// TODO: Implement aliases
 		await Task.CompletedTask;
@@ -194,20 +210,15 @@ public class StorageClient : BundleStorageClient
 
 	public async Task<BlobHandle> WriteRefAsync(RefName name, Bundle bundle, int exportIdx, Utf8String prefix = default, RefOptions? options = null, CancellationToken cancellationToken = default)
 	{
-		BundleLocator locator = await WriteBundleAsync(bundle, prefix, cancellationToken);
-		BlobHandle target = new FlushedNodeHandle(BundleReader, new BundleNodeLocator(bundle.Header.Exports[exportIdx].Hash, locator, exportIdx));
-		await WriteRefTargetAsync(name, target, options, cancellationToken);
+		BundleLocator locator = await this.WriteBundleAsync(bundle, prefix, cancellationToken);
+		BundleNodeLocator nodeLocator = new BundleNodeLocator(bundle.Header.Exports[exportIdx].Hash, locator, exportIdx);
+		await WriteRefTargetAsync(name, nodeLocator, options, cancellationToken);
 
-		return target;
-	}
-
-	public override Task WriteRefTargetAsync(RefName refName, BundleNodeHandle target, RefOptions? requestOptions, CancellationToken cancellationToken)
-	{
-		return WriteRefTargetAsync(refName, target.GetLocator(), requestOptions, cancellationToken);
+		return CreateNodeHandle(nodeLocator);
 	}
 
 #pragma warning disable IDE0060
-	public async Task WriteRefTargetAsync(RefName refName, BundleNodeLocator target, RefOptions? requestOptions, CancellationToken cancellationToken)
+	public override async Task WriteRefTargetAsync(RefName refName, BundleNodeLocator target, RefOptions? requestOptions, CancellationToken cancellationToken)
 #pragma warning restore IDE0060
 	{
 		BlobId bundleBlob = BlobId.FromBlobLocator(target.Blob);
