Unreal Cloud DDC - Fixed build error with GetStats not being implemented

diff --git a/Engine/Source/Programs/UnrealCloudDDC/Jupiter/Implementation/Bundles/StorageService.cs b/Engine/Source/Programs/UnrealCloudDDC/Jupiter/Implementation/Bundles/StorageService.cs
--- a/Engine/Source/Programs/UnrealCloudDDC/Jupiter/Implementation/Bundles/StorageService.cs
+++ b/Engine/Source/Programs/UnrealCloudDDC/Jupiter/Implementation/Bundles/StorageService.cs
@@ -134,6 +134,11 @@ public sealed class JupiterStorageBackend : EpicGames.Horde.Storage.IStorageBack
 		return (locator, redirectUri);
 	}
 
+	public void GetStats(StorageStats stats)
+	{
+		throw new NotImplementedException();
+	}
+
 	public Task WriteExplicitPathAsync(string path, Stream stream, CancellationToken cancellationToken = default) => throw new NotSupportedException();
 }
 
