Unreal Cloud DDC - Specifying version of newtonsoft as we get this from a dependency and it should be a newer version then what we do get.

diff --git a/Engine/Source/Programs/UnrealCloudDDC/Jupiter.Common/Jupiter.Common.csproj b/Engine/Source/Programs/UnrealCloudDDC/Jupiter.Common/Jupiter.Common.csproj
--- a/Engine/Source/Programs/UnrealCloudDDC/Jupiter.Common/Jupiter.Common.csproj
+++ b/Engine/Source/Programs/UnrealCloudDDC/Jupiter.Common/Jupiter.Common.csproj
@@ -19,6 +19,7 @@
     <PackageReference Include="AWSSDK.SecurityToken" Version="3.7.1.4" />
     <PackageReference Include="Azure.Identity" Version="1.5.0" />
     <PackageReference Include="Azure.Security.KeyVault.Secrets" Version="4.2.0" />
+    <PackageReference Include="Newtonsoft.Json" Version="13.0.1" />
     <PackageReference Include="OpenTelemetry" Version="1.4.0-beta.3" />
     <PackageReference Include="OpenTelemetry.Exporter.Console" Version="1.3.2" />
     <PackageReference Include="OpenTelemetry.Exporter.OpenTelemetryProtocol" Version="1.4.0-beta.3" />
diff --git a/Engine/Source/Programs/UnrealCloudDDC/Jupiter/Jupiter.csproj b/Engine/Source/Programs/UnrealCloudDDC/Jupiter/Jupiter.csproj
--- a/Engine/Source/Programs/UnrealCloudDDC/Jupiter/Jupiter.csproj
+++ b/Engine/Source/Programs/UnrealCloudDDC/Jupiter/Jupiter.csproj
@@ -79,6 +79,7 @@
     <PackageReference Include="Microsoft.Extensions.Caching.Memory" Version="6.0.1" />
     <PackageReference Include="Microsoft.Extensions.Logging.Debug" Version="6.0.0" />
     <PackageReference Include="MongoDB.Driver" Version="2.11.6" />
+    <PackageReference Include="Newtonsoft.Json" Version="13.0.1" />
     <PackageReference Include="SerilogAnalyzer" Version="0.15.0" />
     <PackageReference Include="Crc32.NET" Version="1.2.0" />
     <PackageReference Include="K4os.Compression.LZ4" Version="1.2.16" />
