UnrealCloudDDC: Fix build errors caused by EpicGames.Horde changes.

diff --git a/Engine/Source/Programs/UnrealCloudDDC/Jupiter/Controllers/BundlesController.cs b/Engine/Source/Programs/UnrealCloudDDC/Jupiter/Controllers/BundlesController.cs
--- a/Engine/Source/Programs/UnrealCloudDDC/Jupiter/Controllers/BundlesController.cs
+++ b/Engine/Source/Programs/UnrealCloudDDC/Jupiter/Controllers/BundlesController.cs
@@ -297,9 +297,9 @@ namespace Jupiter.Controllers
 			StorageClient client = await _storageService.GetClientAsync(namespaceId, cancellationToken);
 
 			FindNodesResponse response = new FindNodesResponse();
-			await foreach (BundleNodeHandle handle in client.FindAliasAsync(alias, cancellationToken))
+			foreach (BlobAlias blobAlias in await client.FindAliasesAsync(alias, null, cancellationToken))
 			{
-				response.Nodes.Add(new FindNodeResponse(handle));
+				response.Nodes.Add(new FindNodeResponse((BundleNodeHandle)blobAlias.Target));
 			}
 
 			if (response.Nodes.Count == 0)
diff --git a/Engine/Source/Programs/UnrealCloudDDC/Jupiter/Implementation/Bundles/StorageService.cs b/Engine/Source/Programs/UnrealCloudDDC/Jupiter/Implementation/Bundles/StorageService.cs
--- a/Engine/Source/Programs/UnrealCloudDDC/Jupiter/Implementation/Bundles/StorageService.cs
+++ b/Engine/Source/Programs/UnrealCloudDDC/Jupiter/Implementation/Bundles/StorageService.cs
@@ -169,7 +169,7 @@ public class StorageClient : BundleStorageClient
 		_blobIndex = blobIndex;
 	}
 
-	public override async Task AddAliasAsync(Utf8String name, BundleNodeLocator locator, int rank, CancellationToken cancellationToken = default)
+	public override async Task AddAliasAsync(Utf8String name, BundleNodeLocator locator, int rank, ReadOnlyMemory<byte> data, CancellationToken cancellationToken = default)
 	{
 		// TODO: Implement aliases
 		await Task.CompletedTask;
@@ -181,11 +181,11 @@ public class StorageClient : BundleStorageClient
 		await Task.CompletedTask;
 	}
 
-	public override async IAsyncEnumerable<BundleNodeHandle> FindAliasAsync(Utf8String name, [EnumeratorCancellation] CancellationToken cancellationToken = default)
+	public override async Task<BlobAlias[]> FindAliasesAsync(Utf8String name, int? maxResults = null, CancellationToken cancellationToken = default)
 	{
 		// TODO: Implement aliases
 		await Task.CompletedTask;
-		yield break;
+		throw new NotImplementedException();
 	}
 
 	public override async Task<BundleNodeHandle?> TryReadRefTargetAsync(RefName name, RefCacheTime cacheTime = default, CancellationToken cancellationToken = default)
diff --git a/Engine/Source/Programs/UnrealCloudDDC/Jupiter/Jupiter.csproj b/Engine/Source/Programs/UnrealCloudDDC/Jupiter/Jupiter.csproj
--- a/Engine/Source/Programs/UnrealCloudDDC/Jupiter/Jupiter.csproj
+++ b/Engine/Source/Programs/UnrealCloudDDC/Jupiter/Jupiter.csproj
@@ -41,11 +41,13 @@
   <PropertyGroup Condition="'$(Configuration)|$(Platform)'=='Debug|AnyCPU'">
     <DocumentationFile>Jupiter.xml</DocumentationFile>
     <NoWarn>1701;1702;1591</NoWarn>
+    <WarningsNotAsErrors>NU1902;NU1903</WarningsNotAsErrors>
   </PropertyGroup>
 
   <PropertyGroup Condition="'$(Configuration)|$(Platform)'=='Release|AnyCPU'">
     <DocumentationFile>Jupiter.xml</DocumentationFile>
     <NoWarn>1701;1702;1591</NoWarn>
+    <WarningsNotAsErrors>NU1902;NU1903</WarningsNotAsErrors>
   </PropertyGroup>
 
   <ItemGroup>
