Unreal Cloud DDC - Delete last access tracking of object when we are deleting the object. Fixes issue were we attempt to delete object that doesn't exist.

diff --git a/Engine/Source/Programs/UnrealCloudDDC/Jupiter/Implementation/Objects/ScyllaReferencesStore.cs b/Engine/Source/Programs/UnrealCloudDDC/Jupiter/Implementation/Objects/ScyllaReferencesStore.cs
--- a/Engine/Source/Programs/UnrealCloudDDC/Jupiter/Implementation/Objects/ScyllaReferencesStore.cs
+++ b/Engine/Source/Programs/UnrealCloudDDC/Jupiter/Implementation/Objects/ScyllaReferencesStore.cs
@@ -370,8 +370,12 @@ namespace Jupiter.Implementation
 		{
 			using TelemetrySpan scope = _tracer.BuildScyllaSpan("scylla.delete_record").SetAttribute("resource.name", $"{ns}.{bucket}.{key}");
 
+			Task? lastAccessDeleteTask = _mapper.DeleteAsync<ScyllaObjectLastAccess>("WHERE namespace=? AND bucket=? AND name=?", ns.ToString(), bucket.ToString(), key.ToString());
+
 			AppliedInfo<ScyllaObject> info = await _mapper.DeleteIfAsync<ScyllaObject>("WHERE namespace=? AND bucket=? AND name=? IF EXISTS", ns.ToString(), bucket.ToString(), key.ToString());
 
+			await lastAccessDeleteTask;
+
 			if (info.Applied)
 			{
 				return true;
