Unreal Cloud DDC - Fixed issue listing ref buckets in Scylla

diff --git a/Engine/Source/Programs/UnrealCloudDDC/Jupiter/Implementation/Objects/ScyllaReferencesStore.cs b/Engine/Source/Programs/UnrealCloudDDC/Jupiter/Implementation/Objects/ScyllaReferencesStore.cs
--- a/Engine/Source/Programs/UnrealCloudDDC/Jupiter/Implementation/Objects/ScyllaReferencesStore.cs
+++ b/Engine/Source/Programs/UnrealCloudDDC/Jupiter/Implementation/Objects/ScyllaReferencesStore.cs
@@ -355,7 +355,7 @@ namespace Jupiter.Implementation
 
 		public async IAsyncEnumerable<BucketId> GetBuckets(NamespaceId ns)
 		{
-			foreach (ScyllaBucket? scyllaBucket in await _mapper.FetchAsync<ScyllaBucket>("WHERE ns = ?", ns))
+			foreach (ScyllaBucket? scyllaBucket in await _mapper.FetchAsync<ScyllaBucket>("WHERE namespace=?", ns))
 			{
 				if (scyllaBucket == null)
 				{
