Unreal Cloud DDC - polished some logging and error handling from running metrics calculations

diff --git a/Engine/Source/Programs/UnrealCloudDDC/Jupiter/Implementation/MetricsService.cs b/Engine/Source/Programs/UnrealCloudDDC/Jupiter/Implementation/MetricsService.cs
--- a/Engine/Source/Programs/UnrealCloudDDC/Jupiter/Implementation/MetricsService.cs
+++ b/Engine/Source/Programs/UnrealCloudDDC/Jupiter/Implementation/MetricsService.cs
@@ -93,7 +93,8 @@ namespace Jupiter.Implementation
 					await foreach (BucketId bucket in _referencesStore.GetBuckets(ns).WithCancellation(cancellationToken))
 					{
 						DateTime start = DateTime.UtcNow;
-						
+						_logger.LogInformation("Calculating stats for {Namespace} {Bucket}", ns, bucket);
+
 						await calculator.CalculateStatsForBucketAsync(ns, bucket);
 
 						TimeSpan duration = DateTime.UtcNow - start;
diff --git a/Engine/Source/Programs/UnrealCloudDDC/Jupiter/Implementation/Objects/ScyllaReferencesStore.cs b/Engine/Source/Programs/UnrealCloudDDC/Jupiter/Implementation/Objects/ScyllaReferencesStore.cs
--- a/Engine/Source/Programs/UnrealCloudDDC/Jupiter/Implementation/Objects/ScyllaReferencesStore.cs
+++ b/Engine/Source/Programs/UnrealCloudDDC/Jupiter/Implementation/Objects/ScyllaReferencesStore.cs
@@ -277,6 +277,11 @@ namespace Jupiter.Implementation
 					continue;
 				}
 
+				if (blobIdentifier == null)
+				{
+					continue;
+				}
+
 				yield return (new RefId(name), blobIdentifier.AsBlobIdentifier());
 			}
 		}
