Unreal Cloud DDC - Delete non-finalized refs that are older then 2 hours during ref GC

diff --git a/Engine/Source/Programs/UnrealCloudDDC/Jupiter/Implementation/GC/CacheCleanup.cs b/Engine/Source/Programs/UnrealCloudDDC/Jupiter/Implementation/GC/CacheCleanup.cs
--- a/Engine/Source/Programs/UnrealCloudDDC/Jupiter/Implementation/GC/CacheCleanup.cs
+++ b/Engine/Source/Programs/UnrealCloudDDC/Jupiter/Implementation/GC/CacheCleanup.cs
@@ -91,6 +91,19 @@ namespace Jupiter.Implementation
 				{
 					(NamespaceId ns, BucketId bucket, RefId name, DateTime lastAccessTime) = tuple;
 
+					// if a object was accessed within the last two hours we will let it live even if its un-finalized as it might be written to right now
+					if (lastAccessTime < DateTime.Now.AddHours(-2))
+					{
+						RefRecord refRecord = await _referencesStore.GetAsync(ns, bucket, name, IReferencesStore.FieldFlags.None, IReferencesStore.OperationFlags.None);
+						if (!refRecord.IsFinalized)
+						{
+							_logger.LogInformation("Deleting object {Namespace} {Bucket} {Name} as it is not finalized", ns, bucket, name);
+
+							await DeleteRefAsync(ns, bucket, name);
+							return;
+						}
+					}
+					
 					if (!ShouldGCNamespace(ns))
 					{
 						return;
@@ -109,28 +122,9 @@ namespace Jupiter.Implementation
 					_logger.LogInformation(
 						"Attempting to delete object {Namespace} {Bucket} {Name} as it was last updated {LastAccessTime} which is older then {CutoffTime}",
 						ns, bucket, name, lastAccessTime, cutoffTime);
-					using TelemetrySpan scope = _tracer.StartActiveSpan("gc.ref")
-						.SetAttribute("operation.name", "gc.ref")
-						.SetAttribute("resource.name", $"{ns}:{bucket}.{name}")
-						.SetAttribute("namespace", ns.ToString());
-					// delete the old record from the ref refs
-
-					bool storeDelete = false;
-					try
-					{
-						storeDelete = await _referencesStore.DeleteAsync(ns, bucket, name);
-						if (storeDelete && _settings.CurrentValue.WriteDeleteToReplicationLog)
-						{
-							// insert a delete event into the transaction log
-							await _replicationLog.InsertDeleteEventAsync(ns, bucket, name, null);
-						}
-					}
-					catch (Exception e)
-					{
-						_logger.LogWarning(e, "Exception when attempting to delete record {Bucket} {Name} in {Namespace}",
-							bucket, name, ns);
-					}
 
+					bool storeDelete = await DeleteRefAsync(ns, bucket, name);
+					
 					if (storeDelete)
 					{
 						Interlocked.Increment(ref countOfDeletedRecords);
@@ -147,5 +141,32 @@ namespace Jupiter.Implementation
 
 			return countOfDeletedRecords;
 		}
+
+		private async Task<bool> DeleteRefAsync(NamespaceId ns, BucketId bucket, RefId name)
+		{
+			using TelemetrySpan scope = _tracer.StartActiveSpan("gc.ref")
+				.SetAttribute("operation.name", "gc.ref")
+				.SetAttribute("resource.name", $"{ns}:{bucket}.{name}")
+				.SetAttribute("namespace", ns.ToString());
+			// delete the old record from the ref refs
+
+			bool storeDelete = false;
+			try
+			{
+				storeDelete = await _referencesStore.DeleteAsync(ns, bucket, name);
+				if (storeDelete && _settings.CurrentValue.WriteDeleteToReplicationLog)
+				{
+					// insert a delete event into the transaction log
+					await _replicationLog.InsertDeleteEventAsync(ns, bucket, name, null);
+				}
+			}
+			catch (Exception e)
+			{
+				_logger.LogWarning(e, "Exception when attempting to delete record {Bucket} {Name} in {Namespace}",
+					bucket, name, ns);
+			}
+
+			return storeDelete;
+		}
 	}
 }
