Unreal Cloud DDC - Ignore missing refs when trying to cleanup non-finalized refs (these are cleaned up later)

diff --git a/Engine/Source/Programs/UnrealCloudDDC/Jupiter/Implementation/GC/CacheCleanup.cs b/Engine/Source/Programs/UnrealCloudDDC/Jupiter/Implementation/GC/CacheCleanup.cs
--- a/Engine/Source/Programs/UnrealCloudDDC/Jupiter/Implementation/GC/CacheCleanup.cs
+++ b/Engine/Source/Programs/UnrealCloudDDC/Jupiter/Implementation/GC/CacheCleanup.cs
@@ -94,13 +94,20 @@ namespace Jupiter.Implementation
 					// if a object was accessed within the last two hours we will let it live even if its un-finalized as it might be written to right now
 					if (lastAccessTime < DateTime.Now.AddHours(-2))
 					{
-						RefRecord refRecord = await _referencesStore.GetAsync(ns, bucket, name, IReferencesStore.FieldFlags.None, IReferencesStore.OperationFlags.None);
-						if (!refRecord.IsFinalized)
+						try
 						{
-							_logger.LogInformation("Deleting object {Namespace} {Bucket} {Name} as it is not finalized", ns, bucket, name);
-
-							await DeleteRefAsync(ns, bucket, name);
-							return;
+							RefRecord refRecord = await _referencesStore.GetAsync(ns, bucket, name, IReferencesStore.FieldFlags.None, IReferencesStore.OperationFlags.None);
+							if (!refRecord.IsFinalized)
+							{
+								_logger.LogInformation("Deleting object {Namespace} {Bucket} {Name} as it is not finalized", ns, bucket, name);
+
+								await DeleteRefAsync(ns, bucket, name);
+								return;
+							}
+						}
+						catch (RefNotFoundException)
+						{
+							// ignore refs that can not be found, will be cleaned up later
 						}
 					}
 					
