Unreal Cloud DDC - Use the simple health checks by default (we have been running like this internally for a long time). The ASP.NET health checks have issues with multiple calls in parallel which can happen under kubernetes or if used in multiple places (lb and k8s for instance).

diff --git a/Engine/Source/Programs/UnrealCloudDDC/Helm/UnrealCloudDDC/templates/deployment.yaml b/Engine/Source/Programs/UnrealCloudDDC/Helm/UnrealCloudDDC/templates/deployment.yaml
--- a/Engine/Source/Programs/UnrealCloudDDC/Helm/UnrealCloudDDC/templates/deployment.yaml
+++ b/Engine/Source/Programs/UnrealCloudDDC/Helm/UnrealCloudDDC/templates/deployment.yaml
@@ -101,7 +101,7 @@ spec:
         {{- if .Values.livenessProbe.enabled }}
         livenessProbe:
           httpGet:
-            path: {{ default "health/liveness" .Values.livenessProbe.path }}
+            path: {{ default "health/live" .Values.livenessProbe.path }}
             port: http
           initialDelaySeconds: {{ .Values.livenessProbe.initialDelaySeconds }}
           periodSeconds: {{ .Values.livenessProbe.periodSeconds }}
@@ -112,7 +112,7 @@ spec:
         {{- if .Values.readinessProbe.enabled }}
         readinessProbe:
           httpGet:
-            path:  {{ default "health/readiness" .Values.readinessProbe.path }}
+            path:  {{ default "health/ready" .Values.readinessProbe.path }}
             port: http
           initialDelaySeconds: {{ .Values.readinessProbe.initialDelaySeconds }}
           periodSeconds: {{ .Values.readinessProbe.periodSeconds }}
@@ -180,7 +180,7 @@ spec:
         {{- if .Values.nginx.livenessProbe.enabled }}
         livenessProbe:
           httpGet:
-            path: {{ default "health/liveness" .Values.nginx.livenessProbe.path }}
+            path: {{ default "health/live" .Values.nginx.livenessProbe.path }}
             port: nginx-http
           initialDelaySeconds: {{ .Values.nginx.livenessProbe.initialDelaySeconds }}
           periodSeconds: {{ .Values.nginx.livenessProbe.periodSeconds }}
@@ -191,7 +191,7 @@ spec:
         {{- if .Values.nginx.readinessProbe.enabled }}
         readinessProbe:
           httpGet:
-            path:  {{ default "health/readiness" .Values.nginx.readinessProbe.path }}
+            path:  {{ default "health/ready" .Values.nginx.readinessProbe.path }}
             port: nginx-http
           initialDelaySeconds: {{ .Values.nginx.readinessProbe.initialDelaySeconds }}
           periodSeconds: {{ .Values.nginx.readinessProbe.periodSeconds }}
