Unreal Cloud DDC - Change order that refs are cleaned up, start with the old refs and if its considered alive then check if its finalzied (as that requires another db operation to determine) - speeds up GC and reduces pressure on DB.

diff --git a/Engine/Source/Programs/UnrealCloudDDC/Jupiter/Implementation/GC/CacheCleanup.cs b/Engine/Source/Programs/UnrealCloudDDC/Jupiter/Implementation/GC/CacheCleanup.cs
--- a/Engine/Source/Programs/UnrealCloudDDC/Jupiter/Implementation/GC/CacheCleanup.cs
+++ b/Engine/Source/Programs/UnrealCloudDDC/Jupiter/Implementation/GC/CacheCleanup.cs
@@ -91,6 +91,35 @@ namespace Jupiter.Implementation
 				{
 					(NamespaceId ns, BucketId bucket, RefId name, DateTime lastAccessTime) = tuple;
 
+					if (!ShouldGCNamespace(ns))
+					{
+						return;
+					}
+
+					_logger.LogDebug(
+						"Considering object in {Namespace} {Bucket} {Name} for deletion, was last updated {LastAccessTime}",
+						ns, bucket, name, lastAccessTime);
+					Interlocked.Increment(ref consideredCount);
+
+					if (lastAccessTime < cutoffTime)
+					{
+						_logger.LogInformation(
+							"Attempting to delete object {Namespace} {Bucket} {Name} as it was last updated {LastAccessTime} which is older then {CutoffTime}",
+							ns, bucket, name, lastAccessTime, cutoffTime);
+
+						bool storeDelete = await DeleteRefAsync(ns, bucket, name);
+					
+						if (storeDelete)
+						{
+							Interlocked.Increment(ref countOfDeletedRecords);
+						}
+						else
+						{
+							_logger.LogWarning("Failed to delete record {Bucket} {Name} in {Namespace}", bucket, name, ns);
+						}
+						return;
+					}
+
 					// if a object was accessed within the last two hours we will let it live even if its un-finalized as it might be written to right now
 					if (lastAccessTime < DateTime.Now.AddHours(-2))
 					{
@@ -110,36 +139,6 @@ namespace Jupiter.Implementation
 							// ignore refs that can not be found, will be cleaned up later
 						}
 					}
-					
-					if (!ShouldGCNamespace(ns))
-					{
-						return;
-					}
-
-					_logger.LogDebug(
-						"Considering object in {Namespace} {Bucket} {Name} for deletion, was last updated {LastAccessTime}",
-						ns, bucket, name, lastAccessTime);
-					Interlocked.Increment(ref consideredCount);
-
-					if (lastAccessTime > cutoffTime)
-					{
-						return;
-					}
-
-					_logger.LogInformation(
-						"Attempting to delete object {Namespace} {Bucket} {Name} as it was last updated {LastAccessTime} which is older then {CutoffTime}",
-						ns, bucket, name, lastAccessTime, cutoffTime);
-
-					bool storeDelete = await DeleteRefAsync(ns, bucket, name);
-					
-					if (storeDelete)
-					{
-						Interlocked.Increment(ref countOfDeletedRecords);
-					}
-					else
-					{
-						_logger.LogWarning("Failed to delete record {Bucket} {Name} in {Namespace}", bucket, name, ns);
-					}
 				});
 
 			TimeSpan cleanupDuration = DateTime.Now - cleanupStart;
