Unreal Cloud DDC - Improved logging when we are unable to determine if a blob exists remotley due to missing settings, also avoid reading the settings if we can be certain it exists somewhere.

diff --git a/Engine/Source/Programs/UnrealCloudDDC/Jupiter/Implementation/Blob/BlobService.cs b/Engine/Source/Programs/UnrealCloudDDC/Jupiter/Implementation/Blob/BlobService.cs
--- a/Engine/Source/Programs/UnrealCloudDDC/Jupiter/Implementation/Blob/BlobService.cs
+++ b/Engine/Source/Programs/UnrealCloudDDC/Jupiter/Implementation/Blob/BlobService.cs
@@ -684,11 +684,24 @@ public class BlobService : IBlobService
 
 		using ServerTimingMetricScoped? serverTimingScope = serverTiming?.CreateServerTimingMetricScope("blob.exists-remote", "Verify if blob exists in remotes");
 
-		IOptions<JupiterSettings> jupiterSettings = _httpContextAccessor.HttpContext?.RequestServices.GetService<IOptions<JupiterSettings>>()!;
 		List<string> regions = await _blobIndex.GetBlobRegionsAsync(ns, blob);
 
 		// we do not actually verify that the blob exists remotely as that would take a lot of time
 		// instead we simply check if there are any regions were the blob exists that is not our current region
+
+		// if it exists in more then one region, we are sure it exists somewhere that is not here
+		if (regions.Count > 1)
+		{
+			return true;
+		}
+
+		IOptions<JupiterSettings>? jupiterSettings = _httpContextAccessor.HttpContext?.RequestServices.GetService<IOptions<JupiterSettings>>();
+		if (jupiterSettings == null)
+		{
+			_logger.LogWarning("Unable to fetch current site settings, thus not able to determine if {Blob} exists in the remote regions in {Namespace}. Make sure to set the CurrentSite option.", blob, ns);
+			return false;
+		}
+
 		if (regions.Any(region => !string.Equals(region, jupiterSettings.Value.CurrentSite, StringComparison.OrdinalIgnoreCase)))
 		{
 			return true;
