Unreal Cloud DDC - Made the metrics service more resillient to errors as its likely to encounter some malformed data when scanning through everything like this.

diff --git a/Engine/Source/Programs/UnrealCloudDDC/Jupiter/Implementation/MetricsService.cs b/Engine/Source/Programs/UnrealCloudDDC/Jupiter/Implementation/MetricsService.cs
--- a/Engine/Source/Programs/UnrealCloudDDC/Jupiter/Implementation/MetricsService.cs
+++ b/Engine/Source/Programs/UnrealCloudDDC/Jupiter/Implementation/MetricsService.cs
@@ -166,28 +166,42 @@ namespace Jupiter.Implementation
 				{
 					countOfRefsInBucket += 1;
 
-					BlobContents blobContents = await _blobService.GetObjectAsync(ns, blobId);
-					byte[] rawBlob = await blobContents.Stream.ToByteArrayAsync();
-					CbObject cbObject = new CbObject(rawBlob);
-					// enumerate all referenced blobs from the ref
-					await foreach (BlobId blob in _referenceResolver.GetReferencedBlobs(ns, cbObject))
+					try
 					{
-						// check to see if we have counted this blob before
-						bool added = alreadyCountedBlobs.Add(blob.HashData);
-						if (added)
+						BlobContents blobContents = await _blobService.GetObjectAsync(ns, blobId);
+						byte[] rawBlob = await blobContents.Stream.ToByteArrayAsync();
+						CbObject cbObject = new CbObject(rawBlob);
+						// enumerate all referenced blobs from the ref
+						await foreach (BlobId blob in _referenceResolver.GetReferencedBlobs(ns, cbObject))
 						{
-							// new blob, lets count it
-							BlobContents referencedBlob = await _blobService.GetObjectAsync(ns, blob);
-							sizeOfBlobsInBucket += referencedBlob.Length;
-
-							smallestBlobFound = Math.Min(referencedBlob.Length, smallestBlobFound);
-							largestBlobFound = Math.Max(referencedBlob.Length, largestBlobFound);
-
-							countOfBlobsInBucket += 1;
-
-							_blobSizeHistogram.Record(referencedBlob.Length, tags);
+							// check to see if we have counted this blob before
+							bool added = alreadyCountedBlobs.Add(blob.HashData);
+							if (added)
+							{
+								// new blob, lets count it
+								try
+								{
+									BlobContents referencedBlob = await _blobService.GetObjectAsync(ns, blob);
+									sizeOfBlobsInBucket += referencedBlob.Length;
+
+									smallestBlobFound = Math.Min(referencedBlob.Length, smallestBlobFound);
+									largestBlobFound = Math.Max(referencedBlob.Length, largestBlobFound);
+
+									countOfBlobsInBucket += 1;
+
+									_blobSizeHistogram.Record(referencedBlob.Length, tags);
+								}
+								catch (BlobNotFoundException)
+								{
+									// if one of the referenced blobs is missing we just ignore this particular blob
+								}
+							}
 						}
 					}
+					catch (Exception e)
+					{
+						_logger.LogWarning("Unknown exception {Message} when attempting to calculate metrics for {Namespace} {Bucket}. Ignoring.", e.Message, ns, bucket);
+					}
 				}
 
 				double avgBlobSize = sizeOfBlobsInBucket / (double)countOfBlobsInBucket;
