Unreal Cloud DDC - Bumped batch size when cleaning the filesystem store, for most nodes this doesn't matter as they will only cleanup as many files as needed to reach the target. But for really large nodes, and esspecially if they then do see a lot of new files all the time, this low number can result in the GC never reaching target and thus the cache grows to large over time.

diff --git a/Engine/Source/Programs/UnrealCloudDDC/Jupiter/Implementation/Blob/FileSystemStore.cs b/Engine/Source/Programs/UnrealCloudDDC/Jupiter/Implementation/Blob/FileSystemStore.cs
--- a/Engine/Source/Programs/UnrealCloudDDC/Jupiter/Implementation/Blob/FileSystemStore.cs
+++ b/Engine/Source/Programs/UnrealCloudDDC/Jupiter/Implementation/Blob/FileSystemStore.cs
@@ -155,7 +155,7 @@ namespace Jupiter.Implementation
 		/// <param name="cancellationToken">Cancellation token</param>
 		/// <param name="batchSize">Number of files to scan for clean up. A higher number is recommended since blob store can contain many small blobs</param>
 		/// <returns></returns>
-		public async Task<ulong> CleanupInternalAsync(CancellationToken cancellationToken, int batchSize = 100000)
+		public async Task<ulong> CleanupInternalAsync(CancellationToken cancellationToken, int batchSize = 1_000_000)
 		{
 			using TelemetrySpan scope = _tracer.StartActiveSpan("gc.filesystem")
 				.SetAttribute("operation.name", "gc.filesystem");
