Unreal Cloud DDC - Set worker deployments to use recreate strategy (means we do not attempt to run new and old at the same time when upgrading, this makes it easier to reason about utilization and since these are only for background tasks we do not need to guaranteee uptime)

diff --git a/Engine/Source/Programs/UnrealCloudDDC/Helm/UnrealCloudDDC/templates/deployment-worker.yaml b/Engine/Source/Programs/UnrealCloudDDC/Helm/UnrealCloudDDC/templates/deployment-worker.yaml
--- a/Engine/Source/Programs/UnrealCloudDDC/Helm/UnrealCloudDDC/templates/deployment-worker.yaml
+++ b/Engine/Source/Programs/UnrealCloudDDC/Helm/UnrealCloudDDC/templates/deployment-worker.yaml
@@ -16,6 +16,8 @@ spec:
       app.kubernetes.io/name: {{ include "common.names.name" . }}-worker
       app.kubernetes.io/instance: {{ .Release.Name }}
   replicas: {{ $worker.replicaCount }}
+  strategy:
+   type: Recreate
   template:
     metadata:
       labels: 
