Unreal Cloud DDC - Updated changelog with changes since last release (0.3.0)

diff --git a/Engine/Source/Programs/UnrealCloudDDC/changelog.md b/Engine/Source/Programs/UnrealCloudDDC/changelog.md
--- a/Engine/Source/Programs/UnrealCloudDDC/changelog.md
+++ b/Engine/Source/Programs/UnrealCloudDDC/changelog.md
@@ -3,7 +3,16 @@
 * *Breaking* The scylla connection string now needs to include the default keyspace, a example connection string is  `Contact Points=your-scylla-dns.your-domain.com;;Default Keyspace=jupiter;`. The keyspace is `jupiter` is you are migrating from older releases. This allows you to also set the keyspace to something different if you want to run multiple instances of Unreal Cloud DDC against the same scylla cluster.)
 * Migration options from `0.3.0` have been updated to assume you have migrated by default.
 * Added `prepareNvmeFilesystem` section in Helm chart that creates a initContainer which will format a attached nvme drive.
-
+* Fixed issue when content id remapped due to a new version, its supposed to use the smaller version but was in fact using the larger.
+* Bug fixes to make the speculative replication more resilient.
+* Ability to use nginx to sanatize http traffic and only use a level 4 load balancer in front as a performance improvment.
+* Fixes for handling 2GB+ files.
+* Fixes for on-demand replication.
+* Support for S3 multipart uploads of large files.
+* `Scylla.AvoidSchemaChanges` can now be set to avoid triggering schema modifications - this forces you to manually apply any schema changes required when upgrading but also means you can avoid triggering schema changes while maintinance is running.
+* Added `MetricsService` which is disabled by default, will scan all data to determine things like number of objects in each bucket and sizes. Puts a fairly heavy load on your DB so is disabled by default.
+* Added GC of non-finalized refs at a more aggresive cadence then normal refs (removed when they are 2 hours old)
+  
 # 0.3.0
 * Azure blob storage now supports storage pools
 * Last access table refactoring - Moved the last accessing tracking out of the objects table and into a seperate table. Saves on compation work for Scylla. Set `Scylla.ListObjectsFromLastAccessTable` to migrate GC refs to use this new table (will be default in the next release).
