Unreal Cloud DDC - Added namespace to metrics collected by the metrics service

diff --git a/Engine/Source/Programs/UnrealCloudDDC/Jupiter/Implementation/MetricsService.cs b/Engine/Source/Programs/UnrealCloudDDC/Jupiter/Implementation/MetricsService.cs
--- a/Engine/Source/Programs/UnrealCloudDDC/Jupiter/Implementation/MetricsService.cs
+++ b/Engine/Source/Programs/UnrealCloudDDC/Jupiter/Implementation/MetricsService.cs
@@ -152,7 +152,7 @@ namespace Jupiter.Implementation
 
 		public async Task<BucketStats?> CalculateStatsForBucketAsync(NamespaceId ns, BucketId bucket)
 		{
-			KeyValuePair<string, object?>[] tags = new[] { new KeyValuePair<string, object?>("Bucket", bucket.ToString()) };
+			KeyValuePair<string, object?>[] tags = new[] { new KeyValuePair<string, object?>("Bucket", bucket.ToString()), new KeyValuePair<string, object?>("Namespace", ns.ToString()) };
 
 			try
 			{
