Unreal Cloud DDC - Fixed issue with being unable to determine current region in some worker processes occasionally.

diff --git a/Engine/Source/Programs/UnrealCloudDDC/Jupiter/Implementation/Blob/BlobService.cs b/Engine/Source/Programs/UnrealCloudDDC/Jupiter/Implementation/Blob/BlobService.cs
--- a/Engine/Source/Programs/UnrealCloudDDC/Jupiter/Implementation/Blob/BlobService.cs
+++ b/Engine/Source/Programs/UnrealCloudDDC/Jupiter/Implementation/Blob/BlobService.cs
@@ -41,6 +41,7 @@ public class BlobService : IBlobService
 	private readonly ILogger _logger;
 	private readonly Counter<long>? _storeGetHitsCounter;
 	private readonly Counter<long>? _storeGetAttemptsCounter;
+	private readonly string _currentSite;
 
 	internal IEnumerable<IBlobStore> BlobStore
 	{
@@ -48,7 +49,7 @@ public class BlobService : IBlobService
 		set => _blobStores = value.ToList();
 	}
 
-	public BlobService(IServiceProvider provider, IOptionsMonitor<UnrealCloudDDCSettings> settings, IBlobIndex blobIndex, IPeerStatusService peerStatusService, IHttpClientFactory httpClientFactory, IServiceCredentials serviceCredentials, INamespacePolicyResolver namespacePolicyResolver, IHttpContextAccessor httpContextAccessor, IRequestHelper? requestHelper, Tracer tracer, BufferedPayloadFactory bufferedPayloadFactory, ILogger<BlobService> logger, Meter? meter)
+	public BlobService(IServiceProvider provider, IOptionsMonitor<UnrealCloudDDCSettings> settings, IOptionsMonitor<JupiterSettings> jupiterSettings, IBlobIndex blobIndex, IPeerStatusService peerStatusService, IHttpClientFactory httpClientFactory, IServiceCredentials serviceCredentials, INamespacePolicyResolver namespacePolicyResolver, IHttpContextAccessor httpContextAccessor, IRequestHelper? requestHelper, Tracer tracer, BufferedPayloadFactory bufferedPayloadFactory, ILogger<BlobService> logger, Meter? meter)
 	{
 		_blobStores = GetBlobStores(provider, settings).ToList();
 		_settings = settings;
@@ -63,6 +64,8 @@ public class BlobService : IBlobService
 		_bufferedPayloadFactory = bufferedPayloadFactory;
 		_logger = logger;
 
+		_currentSite = jupiterSettings.CurrentValue.CurrentSite;
+
 		_storeGetHitsCounter = meter?.CreateCounter<long>("store.blob_get.found");
 		_storeGetAttemptsCounter = meter?.CreateCounter<long>("store.blob_get.attempt");
 	}
@@ -695,14 +698,7 @@ public class BlobService : IBlobService
 			return true;
 		}
 
-		IOptions<JupiterSettings>? jupiterSettings = _httpContextAccessor.HttpContext?.RequestServices.GetService<IOptions<JupiterSettings>>();
-		if (jupiterSettings == null)
-		{
-			_logger.LogWarning("Unable to fetch current site settings, thus not able to determine if {Blob} exists in the remote regions in {Namespace}. Make sure to set the CurrentSite option.", blob, ns);
-			return false;
-		}
-
-		if (regions.Any(region => !string.Equals(region, jupiterSettings.Value.CurrentSite, StringComparison.OrdinalIgnoreCase)))
+		if (regions.Any(region => !string.Equals(region, _currentSite, StringComparison.OrdinalIgnoreCase)))
 		{
 			return true;
 		}
