Unreal Cloud DDC - Added open telemetry metric for number of considered refs for GC.

diff --git a/Engine/Source/Programs/UnrealCloudDDC/Jupiter/Implementation/GC/CacheCleanup.cs b/Engine/Source/Programs/UnrealCloudDDC/Jupiter/Implementation/GC/CacheCleanup.cs
--- a/Engine/Source/Programs/UnrealCloudDDC/Jupiter/Implementation/GC/CacheCleanup.cs
+++ b/Engine/Source/Programs/UnrealCloudDDC/Jupiter/Implementation/GC/CacheCleanup.cs
@@ -1,6 +1,8 @@
 ﻿// Copyright Epic Games, Inc. All Rights Reserved.
 
 using System;
+using System.Collections.Generic;
+using System.Diagnostics.Metrics;
 using System.Threading;
 using System.Threading.Tasks;
 using EpicGames.Horde.Storage;
@@ -24,9 +26,10 @@ namespace Jupiter.Implementation
 		private readonly INamespacePolicyResolver _namespacePolicyResolver;
 		private readonly Tracer _tracer;
 		private readonly ILogger _logger;
+		private readonly Gauge<long> _cleanupRefsConsidered;
 
 		public RefLastAccessCleanup(IOptionsMonitor<GCSettings> settings, IReferencesStore referencesStore,
-			IReplicationLog replicationLog, INamespacePolicyResolver namespacePolicyResolver, Tracer tracer, ILogger<RefLastAccessCleanup> logger)
+			IReplicationLog replicationLog, INamespacePolicyResolver namespacePolicyResolver, Meter meter, Tracer tracer, ILogger<RefLastAccessCleanup> logger)
 		{
 			_settings = settings;
 			_referencesStore = referencesStore;
@@ -34,6 +37,8 @@ namespace Jupiter.Implementation
 			_namespacePolicyResolver = namespacePolicyResolver;
 			_tracer = tracer;
 			_logger = logger;
+
+			_cleanupRefsConsidered = meter.CreateGauge<long>("refs.considered");
 		}
 
 		private bool ShouldGCNamespace(NamespaceId ns)
@@ -79,7 +84,7 @@ namespace Jupiter.Implementation
 		{
 			int countOfDeletedRecords = 0;
 			DateTime cutoffTime = DateTime.Now.AddSeconds(-1 * _settings.CurrentValue.LastAccessCutoff.TotalSeconds);
-			ulong consideredCount = 0;
+			long consideredCount = 0;
 			DateTime cleanupStart = DateTime.Now;
 
 			await Parallel.ForEachAsync(_referencesStore.GetRecordsAsync(),
@@ -99,7 +104,9 @@ namespace Jupiter.Implementation
 					_logger.LogDebug(
 						"Considering object in {Namespace} {Bucket} {Name} for deletion, was last updated {LastAccessTime}",
 						ns, bucket, name, lastAccessTime);
+					
 					Interlocked.Increment(ref consideredCount);
+					_cleanupRefsConsidered.Record(consideredCount, Array.Empty<KeyValuePair<string, object?>>());
 
 					if (lastAccessTime < cutoffTime)
 					{
