Unreal Cloud DDC - Preparing release of version 0.4.0

diff --git a/Engine/Source/Programs/UnrealCloudDDC/Helm/UnrealCloudDDC/Chart.yaml b/Engine/Source/Programs/UnrealCloudDDC/Helm/UnrealCloudDDC/Chart.yaml
--- a/Engine/Source/Programs/UnrealCloudDDC/Helm/UnrealCloudDDC/Chart.yaml
+++ b/Engine/Source/Programs/UnrealCloudDDC/Helm/UnrealCloudDDC/Chart.yaml
@@ -4,7 +4,7 @@ description: A Helm chart for Kubernetes deploying UnrealCloudDDC
 
 type: application
 
-version: 0.3.0
+version: 0.4.0
 
 icon: https://devtools-misc.s3.amazonaws.com/ue4logo.png
 home: https://unrealengine.com
diff --git a/Engine/Source/Programs/UnrealCloudDDC/Helm/UnrealCloudDDC/version.yaml b/Engine/Source/Programs/UnrealCloudDDC/Helm/UnrealCloudDDC/version.yaml
--- a/Engine/Source/Programs/UnrealCloudDDC/Helm/UnrealCloudDDC/version.yaml
+++ b/Engine/Source/Programs/UnrealCloudDDC/Helm/UnrealCloudDDC/version.yaml
@@ -1 +1 @@
-version: 0.3.0
\ No newline at end of file
+version: 0.4.0
\ No newline at end of file
diff --git a/Engine/Source/Programs/UnrealCloudDDC/changelog.md b/Engine/Source/Programs/UnrealCloudDDC/changelog.md
--- a/Engine/Source/Programs/UnrealCloudDDC/changelog.md
+++ b/Engine/Source/Programs/UnrealCloudDDC/changelog.md
@@ -1,5 +1,6 @@
 # Unreleased
 
+# 0.4.0
 * *Breaking* The scylla connection string now needs to include the default keyspace, a example connection string is  `Contact Points=your-scylla-dns.your-domain.com;;Default Keyspace=jupiter;`. The keyspace is `jupiter` is you are migrating from older releases. This allows you to also set the keyspace to something different if you want to run multiple instances of Unreal Cloud DDC against the same scylla cluster.)
 * Migration options from `0.3.0` have been updated to assume you have migrated by default.
 * Added `prepareNvmeFilesystem` section in Helm chart that creates a initContainer which will format a attached nvme drive.
diff --git a/Engine/Source/Programs/UnrealCloudDDC/version.yaml b/Engine/Source/Programs/UnrealCloudDDC/version.yaml
--- a/Engine/Source/Programs/UnrealCloudDDC/version.yaml
+++ b/Engine/Source/Programs/UnrealCloudDDC/version.yaml
@@ -1 +1 @@
-version: 0.3.0
\ No newline at end of file
+version: 0.4.0
\ No newline at end of file
