Unreal Cloud DDC - Tweaked scylla connection settings and clarified what they actually mean in terms of connections to host and max load on them, added ability to increase number of parallell requests to node, which should fix occasional busypool exception when we hit the max.

diff --git a/Engine/Source/Programs/UnrealCloudDDC/Jupiter/Implementation/Objects/ScyllaSettings.cs b/Engine/Source/Programs/UnrealCloudDDC/Jupiter/Implementation/Objects/ScyllaSettings.cs
--- a/Engine/Source/Programs/UnrealCloudDDC/Jupiter/Implementation/Objects/ScyllaSettings.cs
+++ b/Engine/Source/Programs/UnrealCloudDDC/Jupiter/Implementation/Objects/ScyllaSettings.cs
@@ -38,10 +38,15 @@ public class ScyllaSettings : IValidatableObject
 	public string LocalKeyspaceSuffix { get; set; } = null!;
 
 	/// <summary>
-	/// Max number of connections for each scylla host before switching to another host
+	/// Max number of connections opened, so the max amount of scylla nodes connected to.
 	/// See https://docs.datastax.com/en/developer/nodejs-driver/4.6/features/connection-pooling/
 	/// </summary>
-	public int MaxConnectionForLocalHost { get; set; } = 8192;
+	public int MaxConnectionForLocalHost { get; set; } = 8;
+
+	/// <summary>
+	/// Fixed maximum amount of requests in flight for a connection before a new connection is used (thus max requests per node)
+	/// </summary>
+	public int MaxRequestsPerConnection { get; set; } = 4096;
 
 	/// <summary>
 	/// The time for a replication log event to live in the incremental state before being deleted, assumption is that the snapshot will have processed the event within this time
diff --git a/Engine/Source/Programs/UnrealCloudDDC/Jupiter/Jupiter.Startup.cs b/Engine/Source/Programs/UnrealCloudDDC/Jupiter/Jupiter.Startup.cs
--- a/Engine/Source/Programs/UnrealCloudDDC/Jupiter/Jupiter.Startup.cs
+++ b/Engine/Source/Programs/UnrealCloudDDC/Jupiter/Jupiter.Startup.cs
@@ -306,7 +306,7 @@ namespace Jupiter
 			Builder clusterBuilder = Cluster.Builder()
 				.WithConnectionString(connectionString)
 				.WithLoadBalancingPolicy(Policies.NewDefaultLoadBalancingPolicy(settings.LocalDatacenterName))
-				.WithPoolingOptions(PoolingOptions.Create().SetMaxConnectionsPerHost(HostDistance.Local, settings.MaxConnectionForLocalHost))
+				.WithPoolingOptions(PoolingOptions.Create().SetMaxConnectionsPerHost(HostDistance.Local, settings.MaxConnectionForLocalHost).SetMaxRequestsPerConnection(settings.MaxRequestsPerConnection))
 				.WithExecutionProfiles(options =>
 					options.WithProfile("default", builder => builder.WithConsistencyLevel(ConsistencyLevel.LocalOne)));
 
