Unreal Cloud DDC - Added trace scope to the metrics service calculations so that we can see how long each calculation takes and what it does.

diff --git a/Engine/Source/Programs/UnrealCloudDDC/Jupiter/Implementation/MetricsService.cs b/Engine/Source/Programs/UnrealCloudDDC/Jupiter/Implementation/MetricsService.cs
--- a/Engine/Source/Programs/UnrealCloudDDC/Jupiter/Implementation/MetricsService.cs
+++ b/Engine/Source/Programs/UnrealCloudDDC/Jupiter/Implementation/MetricsService.cs
@@ -12,6 +12,7 @@ using Jupiter.Utils;
 using Microsoft.Extensions.DependencyInjection;
 using Microsoft.Extensions.Logging;
 using Microsoft.Extensions.Options;
+using OpenTelemetry.Trace;
 
 namespace Jupiter.Implementation
 {
@@ -128,6 +129,7 @@ namespace Jupiter.Implementation
 		private readonly IReferenceResolver _referenceResolver;
 
 		private readonly ILogger _logger;
+		private readonly Tracer _tracer;
 		private readonly Histogram<long> _blobSizeHistogram;
 		private readonly Gauge<double> _blobSizeAvgGauge;
 		private readonly Gauge<long> _blobSizeMinGauge;
@@ -135,12 +137,13 @@ namespace Jupiter.Implementation
 		private readonly Gauge<long> _refsInBucketGauge;
 		private readonly Gauge<long> _blobSizeCountGauge;
 
-		public MetricsCalculator(IReferencesStore referencesStore, IBlobService blobService, IReferenceResolver referenceResolver, Meter meter, ILogger<MetricsService> logger)
+		public MetricsCalculator(IReferencesStore referencesStore, IBlobService blobService, IReferenceResolver referenceResolver, Meter meter, ILogger<MetricsService> logger, Tracer tracer)
 		{
 			_referencesStore = referencesStore;
 			_blobService = blobService;
 			_referenceResolver = referenceResolver;
 			_logger = logger;
+			_tracer = tracer;
 
 			_blobSizeHistogram = meter.CreateHistogram<long>("blobstats.size");
 			_blobSizeAvgGauge = meter.CreateGauge<double>("blobstats.bucket_size.avg");
@@ -166,6 +169,12 @@ namespace Jupiter.Implementation
 				{
 					countOfRefsInBucket += 1;
 
+					using TelemetrySpan scope = _tracer.StartActiveSpan("metrics.calculate")
+						.SetAttribute("operation.name", "metrics.calculate")
+						.SetAttribute("resource.name", $"{ns}:{bucket}.{blobId}")
+						.SetAttribute("namespace", ns.ToString())
+						.SetAttribute("bucket", bucket.ToString());
+
 					try
 					{
 						BlobContents blobContents = await _blobService.GetObjectAsync(ns, blobId);
