Unreal Cloud DDC - Removed allow filtering statements from scylla queries that did not need them (because no filtering happened)

diff --git a/Engine/Source/Programs/UnrealCloudDDC/Jupiter/Implementation/Objects/ScyllaReferencesStore.cs b/Engine/Source/Programs/UnrealCloudDDC/Jupiter/Implementation/Objects/ScyllaReferencesStore.cs
--- a/Engine/Source/Programs/UnrealCloudDDC/Jupiter/Implementation/Objects/ScyllaReferencesStore.cs
+++ b/Engine/Source/Programs/UnrealCloudDDC/Jupiter/Implementation/Objects/ScyllaReferencesStore.cs
@@ -83,8 +83,8 @@ namespace Jupiter.Implementation
 
 			// BYPASS CACHE is a scylla specific extension to disable populating the cache, should be ignored by other cassandra dbs
 			string cqlOptions = scyllaSessionManager.IsScylla ? "BYPASS CACHE" : "";
-			_getObjectsStatement = _session.Prepare($"SELECT namespace, bucket, name, last_access_time FROM objects ALLOW FILTERING {cqlOptions}");
-			_getObjectsLastAccessStatement = _session.Prepare($"SELECT namespace, bucket, name, last_access_time FROM object_last_access_v2 ALLOW FILTERING {cqlOptions}");
+			_getObjectsStatement = _session.Prepare($"SELECT namespace, bucket, name, last_access_time FROM objects {cqlOptions}");
+			_getObjectsLastAccessStatement = _session.Prepare($"SELECT namespace, bucket, name, last_access_time FROM object_last_access_v2 {cqlOptions}");
 			_getNamespacesStatement = _session.Prepare("SELECT DISTINCT namespace FROM buckets_v2");
 			_getNamespacesOldStatement = _session.Prepare("SELECT DISTINCT namespace FROM buckets");
 
