Unreal Cloud DDC - Increase parallism for ref GC in a effort to process more refs at the same time thus speeding up the rate at which we remove stale refs.

diff --git a/Engine/Source/Programs/UnrealCloudDDC/Jupiter/Settings.cs b/Engine/Source/Programs/UnrealCloudDDC/Jupiter/Settings.cs
--- a/Engine/Source/Programs/UnrealCloudDDC/Jupiter/Settings.cs
+++ b/Engine/Source/Programs/UnrealCloudDDC/Jupiter/Settings.cs
@@ -196,7 +196,7 @@ namespace Jupiter
 		public TimeSpan BlobCleanupPollFrequency { get; set; } = TimeSpan.FromMinutes(60);
 		public TimeSpan RefCleanupPollFrequency { get; set; } = TimeSpan.FromMinutes(60);
 		public int OrphanGCMaxParallelOperations { get; set; } = 8;
-		public int OrphanRefMaxParallelOperations { get; set; } = 8;
+		public int OrphanRefMaxParallelOperations { get; set; } = 16;
 		public bool WriteDeleteToReplicationLog { get; set; } = false;
 		public NamespacePolicy.StoragePoolGCMethod DefaultGCPolicy { get; set; } = NamespacePolicy.StoragePoolGCMethod.LastAccess;
 	}
