Cloud DDC - If BypassCache option is set for Get operation we now tell Scylla to bypass cache as well. This means that exists checks by default will not hit the scylla cache - they are only used by GC and thus will query all records and should bypass the cache (to avoid thrasing it)

diff --git a/Engine/Source/Programs/UnrealCloudDDC/Jupiter/Implementation/GC/OrphanBlobCleanupRefs.cs b/Engine/Source/Programs/UnrealCloudDDC/Jupiter/Implementation/GC/OrphanBlobCleanupRefs.cs
--- a/Engine/Source/Programs/UnrealCloudDDC/Jupiter/Implementation/GC/OrphanBlobCleanupRefs.cs
+++ b/Engine/Source/Programs/UnrealCloudDDC/Jupiter/Implementation/GC/OrphanBlobCleanupRefs.cs
@@ -160,8 +160,7 @@ namespace Jupiter.Implementation
 
 						try
 						{
-							(RefRecord, BlobContents?) _ = await _refService.GetAsync(blobNamespace, bucket, key, new string[] { "name" }, doLastAccessTracking: false);
-							found = true;
+							found = await _refService.ExistsAsync(blobNamespace, bucket, key);
 							break;
 						}
 						catch (RefNotFoundException)
diff --git a/Engine/Source/Programs/UnrealCloudDDC/Jupiter/Implementation/Objects/ScyllaReferencesStore.cs b/Engine/Source/Programs/UnrealCloudDDC/Jupiter/Implementation/Objects/ScyllaReferencesStore.cs
--- a/Engine/Source/Programs/UnrealCloudDDC/Jupiter/Implementation/Objects/ScyllaReferencesStore.cs
+++ b/Engine/Source/Programs/UnrealCloudDDC/Jupiter/Implementation/Objects/ScyllaReferencesStore.cs
@@ -19,6 +19,7 @@ namespace Jupiter.Implementation
 	{
 		private readonly ISession _session;
 		private readonly IMapper _mapper;
+		private readonly IScyllaSessionManager _scyllaSessionManager;
 		private readonly IOptionsMonitor<ScyllaSettings> _settings;
 		private readonly INamespacePolicyResolver _namespacePolicyResolver;
 		private readonly Tracer _tracer;
@@ -36,6 +37,7 @@ namespace Jupiter.Implementation
 		public ScyllaReferencesStore(IScyllaSessionManager scyllaSessionManager, IOptionsMonitor<ScyllaSettings> settings, INamespacePolicyResolver namespacePolicyResolver, Tracer tracer, ILogger<ScyllaReferencesStore> logger)
 		{
 			_session = scyllaSessionManager.GetSessionForReplicatedKeyspace();
+			_scyllaSessionManager = scyllaSessionManager;
 			_settings = settings;
 			_namespacePolicyResolver = namespacePolicyResolver;
 			_tracer = tracer;
@@ -106,8 +108,15 @@ namespace Jupiter.Implementation
 			}
 			else
 			{
+				string cqlOptions = "";
+				if (_scyllaSessionManager.IsScylla && opFlags.HasFlag(IReferencesStore.OperationFlags.BypassCache))
+				{
+					// BYPASS CACHE is a scylla specific extension to disable populating the cache, should be ignored by other cassandra dbs
+					cqlOptions = "BYPASS CACHE";
+				}
+
 				// fetch everything except for the inline blob which is quite large
-				o = await _mapper.SingleOrDefaultAsync<ScyllaObject>("SELECT namespace, bucket, name , payload_hash, is_finalized, last_access_time FROM objects WHERE namespace = ? AND bucket = ? AND name = ?", ns.ToString(), bucket.ToString(), name.ToString());
+				o = await _mapper.SingleOrDefaultAsync<ScyllaObject>($"SELECT namespace, bucket, name , payload_hash, is_finalized, last_access_time FROM objects WHERE namespace = ? AND bucket = ? AND name = ? {cqlOptions}", ns.ToString(), bucket.ToString(), name.ToString());
 			}
 
 			if (o == null)
