Unreal Cloud DDC - Added property to scylla scope if we are bypassing cache to more easily correlate slow queries by design vs slowdowns from the db.

diff --git a/Engine/Source/Programs/UnrealCloudDDC/Jupiter/Implementation/Objects/ScyllaReferencesStore.cs b/Engine/Source/Programs/UnrealCloudDDC/Jupiter/Implementation/Objects/ScyllaReferencesStore.cs
--- a/Engine/Source/Programs/UnrealCloudDDC/Jupiter/Implementation/Objects/ScyllaReferencesStore.cs
+++ b/Engine/Source/Programs/UnrealCloudDDC/Jupiter/Implementation/Objects/ScyllaReferencesStore.cs
@@ -99,6 +99,7 @@ namespace Jupiter.Implementation
 		public async Task<RefRecord> GetAsync(NamespaceId ns, BucketId bucket, RefId name, IReferencesStore.FieldFlags fieldFlags, IReferencesStore.OperationFlags opFlags)
 		{
 			using TelemetrySpan scope = _tracer.BuildScyllaSpan("scylla.get").SetAttribute("resource.name", $"{ns}.{bucket}.{name}");
+			scope.SetAttribute("BypassCache", false);
 
 			ScyllaObject? o;
 			bool includePayload = (fieldFlags & IReferencesStore.FieldFlags.IncludePayload) != 0;
@@ -113,6 +114,7 @@ namespace Jupiter.Implementation
 				{
 					// BYPASS CACHE is a scylla specific extension to disable populating the cache, should be ignored by other cassandra dbs
 					cqlOptions = "BYPASS CACHE";
+					scope.SetAttribute("BypassCache", true);
 				}
 
 				// fetch everything except for the inline blob which is quite large
