Unreal Cloud DDC - Reducing parallelism for ref cleanup as this seems to cause issues with Scylla not responding in a timely way.

diff --git a/Engine/Source/Programs/UnrealCloudDDC/Jupiter/Settings.cs b/Engine/Source/Programs/UnrealCloudDDC/Jupiter/Settings.cs
--- a/Engine/Source/Programs/UnrealCloudDDC/Jupiter/Settings.cs
+++ b/Engine/Source/Programs/UnrealCloudDDC/Jupiter/Settings.cs
@@ -196,7 +196,7 @@ namespace Jupiter
 		public TimeSpan BlobCleanupPollFrequency { get; set; } = TimeSpan.FromMinutes(60);
 		public TimeSpan RefCleanupPollFrequency { get; set; } = TimeSpan.FromMinutes(60);
 		public int OrphanGCMaxParallelOperations { get; set; } = 8;
-		public int OrphanRefMaxParallelOperations { get; set; } = 16;
+		public int OrphanRefMaxParallelOperations { get; set; } = 8;
 		public bool WriteDeleteToReplicationLog { get; set; } = false;
 		public NamespacePolicy.StoragePoolGCMethod DefaultGCPolicy { get; set; } = NamespacePolicy.StoragePoolGCMethod.LastAccess;
 	}
