Unreal Cloud DDC - Skip caches when determining if a object exists, this makes sure the exist check is always consistent

diff --git a/Engine/Source/Programs/UnrealCloudDDC/Jupiter/Implementation/RefService.cs b/Engine/Source/Programs/UnrealCloudDDC/Jupiter/Implementation/RefService.cs
--- a/Engine/Source/Programs/UnrealCloudDDC/Jupiter/Implementation/RefService.cs
+++ b/Engine/Source/Programs/UnrealCloudDDC/Jupiter/Implementation/RefService.cs
@@ -255,7 +255,7 @@ namespace Jupiter.Implementation
 		{
 			try
 			{
-				(RefRecord, BlobContents?) _ = await GetAsync(ns, bucket, key, new string[] {"name"}, doLastAccessTracking: false, skipCache: false);
+				(RefRecord, BlobContents?) _ = await GetAsync(ns, bucket, key, new string[] {"name"}, doLastAccessTracking: false, skipCache: true);
 			}
 			catch (NamespaceNotFoundException)
 			{
