Unreal Cloud DDC - Randomize order we process partitions in when doing GC on refs.
Introduced experimental AllowParallelRecordFetch that will fetch records in parallel.

diff --git a/Engine/Source/Programs/UnrealCloudDDC/Jupiter.Common/Utils/ListExtensions.cs b/Engine/Source/Programs/UnrealCloudDDC/Jupiter.Common/Utils/ListExtensions.cs
new file mode 100644
--- /dev/null
+++ b/Engine/Source/Programs/UnrealCloudDDC/Jupiter.Common/Utils/ListExtensions.cs
@@ -0,0 +1,25 @@
+// Copyright Epic Games, Inc. All Rights Reserved.
+
+using System;
+using System.Collections.Generic;
+
+namespace Jupiter.Common.Utils
+{
+	public static class ListExtensions
+	{
+		public static void Shuffle<T>(this IList<T> list)  
+		{
+			// Randomly shuffle the order of a list by swapping members
+			// based on Fisher Yates shuffle - https://en.wikipedia.org/wiki/Fisher%E2%80%93Yates_shuffle
+			int count = list.Count;
+			while (count > 1) 
+			{
+				count--;
+				int k = Random.Shared.Next(count + 1);  
+				T value = list[k];
+				list[k] = list[count];
+				list[count] = value;
+			}  
+		}
+	}
+}
diff --git a/Engine/Source/Programs/UnrealCloudDDC/Jupiter/Implementation/Objects/ScyllaReferencesStore.cs b/Engine/Source/Programs/UnrealCloudDDC/Jupiter/Implementation/Objects/ScyllaReferencesStore.cs
--- a/Engine/Source/Programs/UnrealCloudDDC/Jupiter/Implementation/Objects/ScyllaReferencesStore.cs
+++ b/Engine/Source/Programs/UnrealCloudDDC/Jupiter/Implementation/Objects/ScyllaReferencesStore.cs
@@ -9,6 +9,7 @@ using Cassandra;
 using Cassandra.Mapping;
 using EpicGames.Horde.Storage;
 using Jupiter.Common;
+using Jupiter.Common.Utils;
 using Microsoft.Extensions.Logging;
 using Microsoft.Extensions.Options;
 using OpenTelemetry.Trace;
@@ -309,25 +310,80 @@ namespace Jupiter.Implementation
 				? _getObjectsLastAccessForPartitionRangeStatement
 				: _getObjectsForPartitionRangeStatement;
 
-			foreach ((long, long) range in ScyllaUtils.GetTableRanges(_settings.CurrentValue.CountOfNodes, _settings.CurrentValue.CountOfCoresPerNode, 3))
+			// generate a list of all the primary key ranges that exist on the cluster
+			List<(long, long)> tableRanges = ScyllaUtils.GetTableRanges(_settings.CurrentValue.CountOfNodes, _settings.CurrentValue.CountOfCoresPerNode, 3).ToList();
+
+			// randomly shuffle this list so that we do not scan them in the same order, means that we will eventually visit all ranges even if the process running this is restarted before we have finished
+			List<int> tableRangeIndices = Enumerable.Range(0, tableRanges.Count).ToList();
+			tableRangeIndices.Shuffle();
+
+			if (_settings.CurrentValue.AllowParallelRecordFetch)
 			{
-				RowSet rowSet = await _session.ExecuteAsync(getObjectStatement.Bind(range.Item1, range.Item2));
-				foreach (Row row in rowSet)
-				{
-					string ns = row.GetValue<string>("namespace");
-					string bucket = row.GetValue<string>("bucket");
-					string name = row.GetValue<string>("name");
-					DateTime? lastAccessTime = row.GetValue<DateTime?>("last_access_time");
+				ConcurrentQueue<(NamespaceId, BucketId, RefId, DateTime)> foundRecords = new ConcurrentQueue<(NamespaceId, BucketId, RefId, DateTime)>();
+
+				Task scanTask = Parallel.ForEachAsync(tableRangeIndices, new ParallelOptions {MaxDegreeOfParallelism = (int)_settings.CurrentValue.CountOfNodes},
+					async (index, token) =>
+					{
+						(long, long) range = tableRanges[index];
+						RowSet rowSet = await _session.ExecuteAsync(getObjectStatement.Bind(range.Item1, range.Item2));
+						foreach (Row row in rowSet)
+						{
+							if (token.IsCancellationRequested)
+							{
+								return;
+							}
+							string ns = row.GetValue<string>("namespace");
+							string bucket = row.GetValue<string>("bucket");
+							string name = row.GetValue<string>("name");
+							DateTime? lastAccessTime = row.GetValue<DateTime?>("last_access_time");
+
+							// skip any names that are not conformant to io hash
+							if (name.Length != 40)
+							{
+								continue;
+							}
+
+							// if last access time is missing we treat it as being very old
+							lastAccessTime ??= DateTime.MinValue;
+							foundRecords.Enqueue((new NamespaceId(ns), new BucketId(bucket), new RefId(name), lastAccessTime.Value));
+						}
+					});
 
-					// skip any names that are not conformant to io hash
-					if (name.Length != 40)
+				while (!scanTask.IsCompleted)
+				{
+					while (foundRecords.TryDequeue(out (NamespaceId, BucketId, RefId, DateTime) foundRecord))
 					{
-						continue;
+						yield return foundRecord;
 					}
 
-					// if last access time is missing we treat it as being very old
-					lastAccessTime ??= DateTime.MinValue;
-					yield return (new NamespaceId(ns), new BucketId(bucket), new RefId(name), lastAccessTime.Value);
+					await Task.Delay(10);
+				}
+
+				await scanTask;
+			}
+			else
+			{
+				foreach (int index in tableRangeIndices)
+				{
+					(long, long) range = tableRanges[index];
+					RowSet rowSet = await _session.ExecuteAsync(getObjectStatement.Bind(range.Item1, range.Item2));
+					foreach (Row row in rowSet)
+					{
+						string ns = row.GetValue<string>("namespace");
+						string bucket = row.GetValue<string>("bucket");
+						string name = row.GetValue<string>("name");
+						DateTime? lastAccessTime = row.GetValue<DateTime?>("last_access_time");
+
+						// skip any names that are not conformant to io hash
+						if (name.Length != 40)
+						{
+							continue;
+						}
+
+						// if last access time is missing we treat it as being very old
+						lastAccessTime ??= DateTime.MinValue;
+						yield return (new NamespaceId(ns), new BucketId(bucket), new RefId(name), lastAccessTime.Value);
+					}
 				}
 			}
 		}
diff --git a/Engine/Source/Programs/UnrealCloudDDC/Jupiter/Implementation/Objects/ScyllaSettings.cs b/Engine/Source/Programs/UnrealCloudDDC/Jupiter/Implementation/Objects/ScyllaSettings.cs
--- a/Engine/Source/Programs/UnrealCloudDDC/Jupiter/Implementation/Objects/ScyllaSettings.cs
+++ b/Engine/Source/Programs/UnrealCloudDDC/Jupiter/Implementation/Objects/ScyllaSettings.cs
@@ -111,6 +111,11 @@ public class ScyllaSettings : IValidatableObject
 	/// </summary>
 	public bool AvoidSchemaChanges { get; set; } = false;
 
+	/// <summary>
+	/// Runs the record fetching in parallel (up to CountOfNodes) - experimental change that only applies when UsePerShardScanning is active 
+	/// </summary>
+	public bool AllowParallelRecordFetch { get; set; } = false;
+
 	public IEnumerable<ValidationResult> Validate(ValidationContext validationContext)
 	{
 		List<ValidationResult> results = new List<ValidationResult>();
