Unreal Cloud DDC - Bypass cache when reading back ref value to check if its finalized - this avoids thrasing the scylla cache while doing the full table scan that is GC.

diff --git a/Engine/Source/Programs/UnrealCloudDDC/Jupiter/Implementation/GC/CacheCleanup.cs b/Engine/Source/Programs/UnrealCloudDDC/Jupiter/Implementation/GC/CacheCleanup.cs
--- a/Engine/Source/Programs/UnrealCloudDDC/Jupiter/Implementation/GC/CacheCleanup.cs
+++ b/Engine/Source/Programs/UnrealCloudDDC/Jupiter/Implementation/GC/CacheCleanup.cs
@@ -132,7 +132,7 @@ namespace Jupiter.Implementation
 					{
 						try
 						{
-							RefRecord refRecord = await _referencesStore.GetAsync(ns, bucket, name, IReferencesStore.FieldFlags.None, IReferencesStore.OperationFlags.None);
+							RefRecord refRecord = await _referencesStore.GetAsync(ns, bucket, name, IReferencesStore.FieldFlags.None, IReferencesStore.OperationFlags.BypassCache);
 							if (!refRecord.IsFinalized)
 							{
 								_logger.LogInformation("Deleting object {Namespace} {Bucket} {Name} as it is not finalized", ns, bucket, name);
