Unreal Cloud DDC - Add more tracing when requesting refs out of the memory cache to debug issues with inconsistent values for textures.

diff --git a/Engine/Source/Programs/UnrealCloudDDC/Jupiter/Implementation/Objects/MemoryCachedReferencesStore.cs b/Engine/Source/Programs/UnrealCloudDDC/Jupiter/Implementation/Objects/MemoryCachedReferencesStore.cs
--- a/Engine/Source/Programs/UnrealCloudDDC/Jupiter/Implementation/Objects/MemoryCachedReferencesStore.cs
+++ b/Engine/Source/Programs/UnrealCloudDDC/Jupiter/Implementation/Objects/MemoryCachedReferencesStore.cs
@@ -7,6 +7,7 @@ using System.Threading.Tasks;
 using EpicGames.Horde.Storage;
 using Microsoft.Extensions.Caching.Memory;
 using Microsoft.Extensions.Options;
+using OpenTelemetry.Trace;
 
 namespace Jupiter.Implementation.Objects
 {
@@ -15,11 +16,13 @@ namespace Jupiter.Implementation.Objects
 		private readonly IReferencesStore _actualStore;
 		private readonly ConcurrentDictionary<NamespaceId, MemoryCache> _referenceCaches = new ConcurrentDictionary<NamespaceId, MemoryCache>();
 		private readonly IOptionsMonitor<MemoryCacheReferencesSettings> _options;
+		private readonly Tracer _tracer;
 
-		public MemoryCachedReferencesStore(IReferencesStore actualStore, IOptionsMonitor<MemoryCacheReferencesSettings> options)
+		public MemoryCachedReferencesStore(IReferencesStore actualStore, IOptionsMonitor<MemoryCacheReferencesSettings> options, Tracer tracer)
 		{
 			_actualStore = actualStore;
 			_options = options;
+			_tracer = tracer;
 		}
 
 		private void AddCacheEntry(NamespaceId ns, BucketId bucket, RefId key, RefRecord record)
@@ -56,13 +59,20 @@ namespace Jupiter.Implementation.Objects
 				return await _actualStore.GetAsync(ns, bucket, key, IReferencesStore.FieldFlags.All, opFlags);
 			}
 
+			using TelemetrySpan scope = _tracer.StartActiveSpan("Ref.get")
+				.SetAttribute("operation.name", "Ref.get")
+				.SetAttribute("resource.name", $"{bucket}.{key}");
+
 			MemoryCache cache = GetCacheForNamespace(ns);
 
 			if (cache.TryGetValue(new CachedReferenceKey(bucket, key), out CachedReferenceEntry cachedResult))
 			{
+				scope.SetAttribute("Found", true);
+				scope.SetAttribute("BlobIdentifier", cachedResult.BlobIdentifier.ToString());
 				return cachedResult.ToRefRecord(fieldFlags);
 			}
 
+			scope.SetAttribute("Found", false);
 			RefRecord objectRecord = await _actualStore.GetAsync(ns, bucket, key, IReferencesStore.FieldFlags.All, opFlags);
 			AddCacheEntry(ns, bucket, key, objectRecord);
 
