Unreal Cloud DDC - Improved tracing when adding a new content id

diff --git a/Engine/Source/Programs/UnrealCloudDDC/Jupiter/Implementation/ContentId/ScyllaContentIdStore.cs b/Engine/Source/Programs/UnrealCloudDDC/Jupiter/Implementation/ContentId/ScyllaContentIdStore.cs
--- a/Engine/Source/Programs/UnrealCloudDDC/Jupiter/Implementation/ContentId/ScyllaContentIdStore.cs
+++ b/Engine/Source/Programs/UnrealCloudDDC/Jupiter/Implementation/ContentId/ScyllaContentIdStore.cs
@@ -123,7 +123,7 @@ namespace Jupiter.Implementation
 
 		public async Task PutAsync(NamespaceId ns, ContentId contentId, BlobId blobIdentifier, int contentWeight)
 		{
-			using TelemetrySpan scope = _tracer.BuildScyllaSpan("ScyllaContentIdStore.PutContentId");
+			using TelemetrySpan scope = _tracer.BuildScyllaSpan("ScyllaContentIdStore.PutContentId").SetAttribute("resource.name", contentId.ToString());
 			if (_scyllaSessionManager.IsScylla)
 			{
 				await _mapper.UpdateAsync<ScyllaContentId>("SET chunks = ? WHERE content_id = ? AND content_weight = ?", new [] {new ScyllaBlobIdentifier(blobIdentifier)}, new ScyllaBlobIdentifier(contentId), contentWeight);
