Unreal Cloud DDC - Fixed bug in the ref memory cache were overwrites would not update the cache (as the overwrites are typically not finalized yet).

diff --git a/Engine/Source/Programs/UnrealCloudDDC/Jupiter/Implementation/Objects/MemoryCachedReferencesStore.cs b/Engine/Source/Programs/UnrealCloudDDC/Jupiter/Implementation/Objects/MemoryCachedReferencesStore.cs
--- a/Engine/Source/Programs/UnrealCloudDDC/Jupiter/Implementation/Objects/MemoryCachedReferencesStore.cs
+++ b/Engine/Source/Programs/UnrealCloudDDC/Jupiter/Implementation/Objects/MemoryCachedReferencesStore.cs
@@ -27,12 +27,6 @@ namespace Jupiter.Implementation.Objects
 
 		private void AddCacheEntry(NamespaceId ns, BucketId bucket, RefId key, RefRecord record)
 		{
-			// we can not cache none finalized records as they will be mutated again when finalized
-			if (!record.IsFinalized)
-			{
-				return;
-			}
-
 			MemoryCache cache = GetCacheForNamespace(ns);
 
 			CachedReferenceEntry cachedEntry = new CachedReferenceEntry(record);
@@ -69,7 +63,12 @@ namespace Jupiter.Implementation.Objects
 			{
 				scope.SetAttribute("Found", true);
 				scope.SetAttribute("BlobIdentifier", cachedResult.BlobIdentifier.ToString());
-				return cachedResult.ToRefRecord(fieldFlags);
+				RefRecord record = cachedResult.ToRefRecord(fieldFlags);
+				if (record.IsFinalized)
+				{
+					// only return finalized records
+					return record;
+				}
 			}
 
 			scope.SetAttribute("Found", false);
@@ -89,9 +88,20 @@ namespace Jupiter.Implementation.Objects
 
 		public Task FinalizeAsync(NamespaceId ns, BucketId bucket, RefId key, BlobId blobIdentifier)
 		{
+			FinalizeCacheEntry(ns, bucket, key);
 			return _actualStore.FinalizeAsync(ns, bucket, key, blobIdentifier);
 		}
 
+		private void FinalizeCacheEntry(NamespaceId ns, BucketId bucket, RefId key)
+		{
+			MemoryCache cache = GetCacheForNamespace(ns);
+
+			if (cache.TryGetValue(new CachedReferenceKey(bucket, key), out CachedReferenceEntry result))
+			{
+				result.IsFinalized = true;
+			}
+		}
+
 		public Task UpdateLastAccessTimeAsync(NamespaceId ns, BucketId bucket, RefId key, DateTime newLastAccessTime)
 		{
 			return _actualStore.UpdateLastAccessTimeAsync(ns, bucket, key, newLastAccessTime);
@@ -205,7 +215,7 @@ namespace Jupiter.Implementation.Objects
 	{
 		private int GetSize()
 		{
-			return Namespace.Text.Text.Length + Bucket.ToString().Length + 20 + Blob?.Length ?? 0 + 20;
+			return Namespace.Text.Text.Length + Bucket.ToString().Length + 20 + Blob?.Length ?? 0 + 20 + 4;
 		}
 
 		public CachedReferenceEntry(RefRecord record)
@@ -215,6 +225,7 @@ namespace Jupiter.Implementation.Objects
 			Name = record.Name;
 			Blob = record.InlinePayload;
 			BlobIdentifier = record.BlobIdentifier;
+			IsFinalized = record.IsFinalized;
 			Size = GetSize();
 		}
 
@@ -224,10 +235,11 @@ namespace Jupiter.Implementation.Objects
 		public byte[]? Blob { get; }
 		public BlobId BlobIdentifier { get; }
 		public int Size { get; }
+		public bool IsFinalized { get; set; }
 
 		public RefRecord ToRefRecord(IReferencesStore.FieldFlags fieldFlags)
 		{
-			return new RefRecord(Namespace, Bucket, Name, DateTime.Now, (fieldFlags & IReferencesStore.FieldFlags.IncludePayload) != 0 ? Blob : null, BlobIdentifier, true);
+			return new RefRecord(Namespace, Bucket, Name, DateTime.Now, (fieldFlags & IReferencesStore.FieldFlags.IncludePayload) != 0 ? Blob : null, BlobIdentifier, IsFinalized);
 		}
 	}
 }
