Unreal Cloud DDC - Set a hard page size limit as there seems to be issues with pagination when combining with token range scan

diff --git a/Engine/Source/Programs/UnrealCloudDDC/Jupiter/Implementation/Objects/ScyllaReferencesStore.cs b/Engine/Source/Programs/UnrealCloudDDC/Jupiter/Implementation/Objects/ScyllaReferencesStore.cs
--- a/Engine/Source/Programs/UnrealCloudDDC/Jupiter/Implementation/Objects/ScyllaReferencesStore.cs
+++ b/Engine/Source/Programs/UnrealCloudDDC/Jupiter/Implementation/Objects/ScyllaReferencesStore.cs
@@ -325,7 +325,9 @@ namespace Jupiter.Implementation
 					async (index, token) =>
 					{
 						(long, long) range = tableRanges[index];
-						RowSet rowSet = await _session.ExecuteAsync(getObjectStatement.Bind(range.Item1, range.Item2));
+						BoundStatement? statement = getObjectStatement.Bind(range.Item1, range.Item2);
+						statement.SetPageSize(5000); // increase page size as there seems to be issues fetching multiple pages when token scanning
+						RowSet rowSet = await _session.ExecuteAsync(statement);
 						foreach (Row row in rowSet)
 						{
 							if (token.IsCancellationRequested)
