Unreal Cloud DDC - Do not warning when failing to delete refs, this usually happens due to a mismatch between last access tracking and objects tables (the ref didn't exist in objects but it did in last access). This will practically end up with the ref not existing anyway so its not really a error condition.

diff --git a/Engine/Source/Programs/UnrealCloudDDC/Jupiter/Implementation/GC/CacheCleanup.cs b/Engine/Source/Programs/UnrealCloudDDC/Jupiter/Implementation/GC/CacheCleanup.cs
--- a/Engine/Source/Programs/UnrealCloudDDC/Jupiter/Implementation/GC/CacheCleanup.cs
+++ b/Engine/Source/Programs/UnrealCloudDDC/Jupiter/Implementation/GC/CacheCleanup.cs
@@ -114,16 +114,10 @@ namespace Jupiter.Implementation
 							"Attempting to delete object {Namespace} {Bucket} {Name} as it was last updated {LastAccessTime} which is older then {CutoffTime}",
 							ns, bucket, name, lastAccessTime, cutoffTime);
 
-						bool storeDelete = await DeleteRefAsync(ns, bucket, name);
+						await DeleteRefAsync(ns, bucket, name);
 					
-						if (storeDelete)
-						{
 							Interlocked.Increment(ref countOfDeletedRecords);
-						}
-						else
-						{
-							_logger.LogWarning("Failed to delete record {Bucket} {Name} in {Namespace}", bucket, name, ns);
-						}
+
 						return;
 					}
 
