Unreal Cloud DDC - Added "NotUsed" Scheme for auth, this can be set to disable JWT type scemes while still enabling auth (the entire thing can be disabled as well). This is mostly useful to still allow auth using ServiceTokens. While not a recommended setup for production usecases it can be a useful way to get a deployment up and running that at least have some auth while still not requiring a OIDC IdP.

diff --git a/Engine/Source/Programs/UnrealCloudDDC/Jupiter.Common/BaseStartup.cs b/Engine/Source/Programs/UnrealCloudDDC/Jupiter.Common/BaseStartup.cs
--- a/Engine/Source/Programs/UnrealCloudDDC/Jupiter.Common/BaseStartup.cs
+++ b/Engine/Source/Programs/UnrealCloudDDC/Jupiter.Common/BaseStartup.cs
@@ -157,6 +157,8 @@ namespace Jupiter
 
 					switch (scheme.Implementation)
 					{
+						case SchemeImplementations.NotUsed:
+							break;
 						case SchemeImplementations.JWTBearer:
 							availableSchemes.Add(name);
 							authenticationBuilder.AddJwtBearer(name, options =>
@@ -550,7 +552,8 @@ namespace Jupiter
 	public enum SchemeImplementations
 	{
 		JWTBearer,
-		Okta
+		Okta,
+		NotUsed
 	};
 
 	public class AuthSchemeEntry: IValidatableObject
