Unreal Cloud DDC - Fixed issue were last access tracking was not being done unless you had explicitly set your gc method to use that. This was never the intention as it was supposed to default to this behavior (for backwards compatability).

diff --git a/Engine/Source/Programs/UnrealCloudDDC/Jupiter/Implementation/LastAccess/LastAccessServiceReferences.cs b/Engine/Source/Programs/UnrealCloudDDC/Jupiter/Implementation/LastAccess/LastAccessServiceReferences.cs
--- a/Engine/Source/Programs/UnrealCloudDDC/Jupiter/Implementation/LastAccess/LastAccessServiceReferences.cs
+++ b/Engine/Source/Programs/UnrealCloudDDC/Jupiter/Implementation/LastAccess/LastAccessServiceReferences.cs
@@ -101,7 +101,8 @@ namespace Jupiter.Implementation
 				return false;
 			}
 
-			return _namespacePolicyResolver.GetPoliciesForNs(ns).GcMethod == NamespacePolicy.StoragePoolGCMethod.LastAccess;
+			NamespacePolicy.StoragePoolGCMethod gcMethod = _namespacePolicyResolver.GetPoliciesForNs(ns).GcMethod ?? NamespacePolicy.StoragePoolGCMethod.LastAccess;
+			return gcMethod == NamespacePolicy.StoragePoolGCMethod.LastAccess;
 		}
 
 		public void Dispose()
diff --git a/Engine/Source/Programs/UnrealCloudDDC/Jupiter/Implementation/Objects/MongoReferencesStore.cs b/Engine/Source/Programs/UnrealCloudDDC/Jupiter/Implementation/Objects/MongoReferencesStore.cs
--- a/Engine/Source/Programs/UnrealCloudDDC/Jupiter/Implementation/Objects/MongoReferencesStore.cs
+++ b/Engine/Source/Programs/UnrealCloudDDC/Jupiter/Implementation/Objects/MongoReferencesStore.cs
@@ -82,7 +82,8 @@ namespace Jupiter.Implementation
 			MongoReferencesModelV0 model = new MongoReferencesModelV0(ns, bucket, key, blobHash, blob, isFinalized, DateTime.Now);
 			
 			NamespacePolicy policy = _namespacePolicyResolver.GetPoliciesForNs(ns);
-			if (policy.GcMethod == NamespacePolicy.StoragePoolGCMethod.TTL)
+			NamespacePolicy.StoragePoolGCMethod gcMethod = policy.GcMethod ?? NamespacePolicy.StoragePoolGCMethod.LastAccess;
+			if (gcMethod == NamespacePolicy.StoragePoolGCMethod.TTL)
 			{
 				model.ExpireAt = DateTime.UtcNow.Add(policy.DefaultTTL);
 			}
diff --git a/Engine/Source/Programs/UnrealCloudDDC/Jupiter/Implementation/Objects/ScyllaReferencesStore.cs b/Engine/Source/Programs/UnrealCloudDDC/Jupiter/Implementation/Objects/ScyllaReferencesStore.cs
--- a/Engine/Source/Programs/UnrealCloudDDC/Jupiter/Implementation/Objects/ScyllaReferencesStore.cs
+++ b/Engine/Source/Programs/UnrealCloudDDC/Jupiter/Implementation/Objects/ScyllaReferencesStore.cs
@@ -155,12 +155,13 @@ namespace Jupiter.Implementation
 
 			int? ttl = null;
 			NamespacePolicy policy = _namespacePolicyResolver.GetPoliciesForNs(ns);
-			if (policy.GcMethod == NamespacePolicy.StoragePoolGCMethod.TTL)
+			NamespacePolicy.StoragePoolGCMethod gcMethod = policy.GcMethod ?? NamespacePolicy.StoragePoolGCMethod.LastAccess;
+			if (gcMethod == NamespacePolicy.StoragePoolGCMethod.TTL)
 			{
 				ttl = (int)policy.DefaultTTL.TotalSeconds;
 			}
 
-			Task? insertLastAccess = policy.GcMethod == NamespacePolicy.StoragePoolGCMethod.LastAccess ? _mapper.InsertAsync<ScyllaObjectLastAccess>(new ScyllaObjectLastAccess(ns, bucket, name, DateTime.Now)) : null;
+			Task? insertLastAccess = gcMethod == NamespacePolicy.StoragePoolGCMethod.LastAccess ? _mapper.InsertAsync<ScyllaObjectLastAccess>(new ScyllaObjectLastAccess(ns, bucket, name, DateTime.Now)) : null;
 
 			await _mapper.InsertAsync<ScyllaObject>(new ScyllaObject(ns, bucket, name, blob, blobHash, isFinalized), ttl: ttl, insertNulls: false);
 
diff --git a/Engine/Source/Programs/UnrealCloudDDC/Jupiter/Implementation/RefService.cs b/Engine/Source/Programs/UnrealCloudDDC/Jupiter/Implementation/RefService.cs
--- a/Engine/Source/Programs/UnrealCloudDDC/Jupiter/Implementation/RefService.cs
+++ b/Engine/Source/Programs/UnrealCloudDDC/Jupiter/Implementation/RefService.cs
@@ -86,7 +86,8 @@ namespace Jupiter.Implementation
 
 			if (doLastAccessTracking)
 			{
-				if (_namespacePolicyResolver.GetPoliciesForNs(ns).GcMethod == NamespacePolicy.StoragePoolGCMethod.LastAccess)
+				NamespacePolicy.StoragePoolGCMethod gcPolicy = _namespacePolicyResolver.GetPoliciesForNs(ns).GcMethod ?? NamespacePolicy.StoragePoolGCMethod.LastAccess;
+				if (gcPolicy == NamespacePolicy.StoragePoolGCMethod.LastAccess)
 				{
 					// we do not wait for the last access tracking as it does not matter when it completes
 					Task lastAccessTask = _lastAccessTracker.TrackUsed(new LastAccessRecord(ns, bucket, key)).ContinueWith((task, _) =>
