Unreal Cloud DDC - Improved logging message when deleting non-finalized ref

diff --git a/Engine/Source/Programs/UnrealCloudDDC/Jupiter/Implementation/GC/CacheCleanup.cs b/Engine/Source/Programs/UnrealCloudDDC/Jupiter/Implementation/GC/CacheCleanup.cs
--- a/Engine/Source/Programs/UnrealCloudDDC/Jupiter/Implementation/GC/CacheCleanup.cs
+++ b/Engine/Source/Programs/UnrealCloudDDC/Jupiter/Implementation/GC/CacheCleanup.cs
@@ -129,7 +129,7 @@ namespace Jupiter.Implementation
 							RefRecord refRecord = await _referencesStore.GetAsync(ns, bucket, name, IReferencesStore.FieldFlags.None, IReferencesStore.OperationFlags.BypassCache);
 							if (!refRecord.IsFinalized)
 							{
-								_logger.LogInformation("Deleting object {Namespace} {Bucket} {Name} as it is not finalized", ns, bucket, name);
+								_logger.LogInformation("Deleting object {Namespace} {Bucket} {Name} as it is not finalized. Was last accessed at {LastAccessTime}", ns, bucket, name, lastAccessTime);
 
 								await DeleteRefAsync(ns, bucket, name);
 								return;
