Unreal Cloud DDC - Reworked integration of service accounts into a seperate scheme, makes it more consistent with other schemes. You can still use this in combination with other schemes.

diff --git a/Engine/Source/Programs/UnrealCloudDDC/Jupiter.Common/BaseStartup.cs b/Engine/Source/Programs/UnrealCloudDDC/Jupiter.Common/BaseStartup.cs
--- a/Engine/Source/Programs/UnrealCloudDDC/Jupiter.Common/BaseStartup.cs
+++ b/Engine/Source/Programs/UnrealCloudDDC/Jupiter.Common/BaseStartup.cs
@@ -157,7 +157,9 @@ namespace Jupiter
 
 					switch (scheme.Implementation)
 					{
-						case SchemeImplementations.NotUsed:
+						case SchemeImplementations.ServiceAccount:
+							availableSchemes.Add(name);
+							authenticationBuilder.AddScheme<ServiceAccountAuthOptions, ServiceAccountAuthHandler>(name, options => { });
 							break;
 						case SchemeImplementations.JWTBearer:
 							availableSchemes.Add(name);
@@ -222,9 +224,6 @@ namespace Jupiter
 				authenticationBuilder.AddTestAuth(options => { });
 			}
 
-			availableSchemes.Add(ServiceAccountAuthHandler.AuthenticationScheme);
-			authenticationBuilder.AddScheme<ServiceAccountAuthOptions, ServiceAccountAuthHandler>(ServiceAccountAuthHandler.AuthenticationScheme, options => { });
-
 			services.AddAuthorization(options =>
 			{
 				options.AddPolicy(NamespaceAccessRequirement.Name, policy =>
@@ -553,7 +552,7 @@ namespace Jupiter
 	{
 		JWTBearer,
 		Okta,
-		NotUsed
+		ServiceAccount
 	};
 
 	public class AuthSchemeEntry: IValidatableObject
