Unreal Cloud DDC - Preparing release 0.5.0

diff --git a/Engine/Source/Programs/UnrealCloudDDC/Helm/UnrealCloudDDC/version.yaml b/Engine/Source/Programs/UnrealCloudDDC/Helm/UnrealCloudDDC/version.yaml
--- a/Engine/Source/Programs/UnrealCloudDDC/Helm/UnrealCloudDDC/version.yaml
+++ b/Engine/Source/Programs/UnrealCloudDDC/Helm/UnrealCloudDDC/version.yaml
@@ -1 +1 @@
-version: 0.4.0
\ No newline at end of file
+version: 0.5.0
\ No newline at end of file
diff --git a/Engine/Source/Programs/UnrealCloudDDC/changelog.md b/Engine/Source/Programs/UnrealCloudDDC/changelog.md
--- a/Engine/Source/Programs/UnrealCloudDDC/changelog.md
+++ b/Engine/Source/Programs/UnrealCloudDDC/changelog.md
@@ -1,5 +1,12 @@
 # Unreleased
 
+# 0.5.0
+* Fixed issue introduced in `0.4.0` that would cause last access tracking to not work correctly.
+* Ability to run with auth enabled but using only the service account scheme - useful for a simpler setup as that does not require setup specific information.
+* Improvments to scylla requests to avoid churning the scylla cache when doing GC.
+* Fixed bug in the ref memory cache were overwrites would not correctly updated the local cache.
+* Tweaks to scylla node connection limits to allow for more requests per connection.
+
 # 0.4.0
 * *Breaking* The scylla connection string now needs to include the default keyspace, a example connection string is  `Contact Points=your-scylla-dns.your-domain.com;;Default Keyspace=jupiter;`. The keyspace is `jupiter` is you are migrating from older releases. This allows you to also set the keyspace to something different if you want to run multiple instances of Unreal Cloud DDC against the same scylla cluster.)
 * Migration options from `0.3.0` have been updated to assume you have migrated by default.
diff --git a/Engine/Source/Programs/UnrealCloudDDC/version.yaml b/Engine/Source/Programs/UnrealCloudDDC/version.yaml
--- a/Engine/Source/Programs/UnrealCloudDDC/version.yaml
+++ b/Engine/Source/Programs/UnrealCloudDDC/version.yaml
@@ -1 +1 @@
-version: 0.4.0
\ No newline at end of file
+version: 0.5.0
\ No newline at end of file
