Unreal Cloud DDC - Added trace scope when putting ref that includes the number of missing blobs, useful when debugging why certain refs are not finalized.

diff --git a/Engine/Source/Programs/UnrealCloudDDC/Jupiter/Controllers/ReferencesController.cs b/Engine/Source/Programs/UnrealCloudDDC/Jupiter/Controllers/ReferencesController.cs
--- a/Engine/Source/Programs/UnrealCloudDDC/Jupiter/Controllers/ReferencesController.cs
+++ b/Engine/Source/Programs/UnrealCloudDDC/Jupiter/Controllers/ReferencesController.cs
@@ -706,9 +706,15 @@ namespace Jupiter.Controllers
 
 			(ContentId[] missingReferences, BlobId[] missingBlobs) = await _refService.PutAsync(ns, bucket, key, blobHeader, payloadObject);
 
-			List<ContentHash> missingHashes = new List<ContentHash>(missingReferences);
-			missingHashes.AddRange(missingBlobs);
-			return Ok(new PutObjectResponse(missingHashes.ToArray()));
+			{
+				using TelemetrySpan scope = _tracer.StartActiveSpan("ref.put").SetAttribute("operation.name", "ref.put");
+
+				List<ContentHash> missingHashes = new List<ContentHash>(missingReferences);
+				missingHashes.AddRange(missingBlobs);
+				ContentHash[] missingArray = missingHashes.ToArray();
+				scope.SetAttribute("NeedsCount", missingArray.Length);
+				return Ok(new PutObjectResponse(missingArray));
+			}
 		}
 
 		[HttpPut("{ns}/{bucket}/{key}", Order = 300)]
