Unreal Cloud DDC - Fixed issue with the default repo incorrectly specifying a version tag, which is also added by the appversion (derived from the helm chart) causing version to exist twice in the resulting value thus making the helm chart unable to find the dockerimages by default.

diff --git a/Engine/Source/Programs/UnrealCloudDDC/Helm/UnrealCloudDDC/values.yaml b/Engine/Source/Programs/UnrealCloudDDC/Helm/UnrealCloudDDC/values.yaml
--- a/Engine/Source/Programs/UnrealCloudDDC/Helm/UnrealCloudDDC/values.yaml
+++ b/Engine/Source/Programs/UnrealCloudDDC/Helm/UnrealCloudDDC/values.yaml
@@ -13,7 +13,7 @@ replicaCount: 2
 
 ## Docker image configuration
 image:
-  repository: "ghcr.io/epicgames/unreal-cloud-ddc:latest"
+  repository: "ghcr.io/epicgames/unreal-cloud-ddc"
   pullPolicy: IfNotPresent
 
 ## Force target Kubernetes version (using Helm capabilites if not set)
