Unreal Cloud DDC - Preparing release 0.5.1

diff --git a/Engine/Source/Programs/UnrealCloudDDC/Helm/UnrealCloudDDC/version.yaml b/Engine/Source/Programs/UnrealCloudDDC/Helm/UnrealCloudDDC/version.yaml
--- a/Engine/Source/Programs/UnrealCloudDDC/Helm/UnrealCloudDDC/version.yaml
+++ b/Engine/Source/Programs/UnrealCloudDDC/Helm/UnrealCloudDDC/version.yaml
@@ -1 +1 @@
-version: 0.5.0
\ No newline at end of file
+version: 0.5.1
\ No newline at end of file
diff --git a/Engine/Source/Programs/UnrealCloudDDC/changelog.md b/Engine/Source/Programs/UnrealCloudDDC/changelog.md
--- a/Engine/Source/Programs/UnrealCloudDDC/changelog.md
+++ b/Engine/Source/Programs/UnrealCloudDDC/changelog.md
@@ -1,5 +1,8 @@
 # Unreleased
 
+# 0.5.1
+* Fixed issue in helm charts defaulting to a incorrect docker registry path.
+
 # 0.5.0
 * Fixed issue introduced in `0.4.0` that would cause last access tracking to not work correctly.
 * Ability to run with auth enabled but using only the service account scheme - useful for a simpler setup as that does not require setup specific information.
diff --git a/Engine/Source/Programs/UnrealCloudDDC/version.yaml b/Engine/Source/Programs/UnrealCloudDDC/version.yaml
--- a/Engine/Source/Programs/UnrealCloudDDC/version.yaml
+++ b/Engine/Source/Programs/UnrealCloudDDC/version.yaml
@@ -1 +1 @@
-version: 0.5.0
\ No newline at end of file
+version: 0.5.1
\ No newline at end of file
