Unreal Cloud DDC - Revert change to use Unreal.MachineName in Cloud DDC, we do not link with EpicGames.Build and the hostname in this case is used within the limited context of kuberentes to identify the machine (so it doesn't need to perfectly describe the machine, just seperate it from other machines connected to the same cluster). This resolves a build error in cloud ddc.

diff --git a/Engine/Source/Programs/UnrealCloudDDC/Jupiter/Implementation/LeaderElection/KubernetesLeaderElection.cs b/Engine/Source/Programs/UnrealCloudDDC/Jupiter/Implementation/LeaderElection/KubernetesLeaderElection.cs
--- a/Engine/Source/Programs/UnrealCloudDDC/Jupiter/Implementation/LeaderElection/KubernetesLeaderElection.cs
+++ b/Engine/Source/Programs/UnrealCloudDDC/Jupiter/Implementation/LeaderElection/KubernetesLeaderElection.cs
@@ -52,7 +52,7 @@ namespace Jupiter.Implementation.LeaderElection
 			KubernetesClientConfiguration config = KubernetesClientConfiguration.InClusterConfig();
 			_client = new Kubernetes(config);
 			
-			_identity = Unreal.MachineName;
+			_identity = System.Net.Dns.GetHostName();
 			_logger.LogInformation("Participating in kubernetes leadership election as {Identity} using {Resource} under {Namespace}", _identity, settings.ConfigMapName, settings.Namespace);
 			_configMapLock = new ConfigMapLock(_client, settings.Namespace, settings.ConfigMapName, _identity);
 
