Unreal Cloud DDC - Add support for setting nginx backlog and tweaking the security context as well as the keep alive connection count (which I also bumped the default value of)

diff --git a/Engine/Source/Programs/UnrealCloudDDC/Helm/UnrealCloudDDC/templates/configmap-nginx.yaml b/Engine/Source/Programs/UnrealCloudDDC/Helm/UnrealCloudDDC/templates/configmap-nginx.yaml
--- a/Engine/Source/Programs/UnrealCloudDDC/Helm/UnrealCloudDDC/templates/configmap-nginx.yaml
+++ b/Engine/Source/Programs/UnrealCloudDDC/Helm/UnrealCloudDDC/templates/configmap-nginx.yaml
@@ -29,12 +29,12 @@ data:
       upstream jupiter_backend {
         zone upstreams 64K;
         server {{ if .Values.nginx.useDomainSockets -}} unix:/nginx/jupiter-http.sock {{- else -}} localhost:80 {{- end }} max_fails=1 fail_timeout=2s;
-        keepalive 64;
+        keepalive {{ .Values.nginx.keepAliveConnections }};
       }
 
       # http 1.1 public port
       server {
-        listen 81 backlog=4096;
+        listen 81 backlog={{ .Values.nginx.backlog }};
         server_name public_http1;
 
         proxy_buffering off;
@@ -78,7 +78,7 @@ data:
 
       # http 1.1 corp port
       server {
-        listen 8009 backlog=4096;
+        listen 8009 backlog={{ .Values.nginx.backlog }};
         server_name corp_http1;
         
         proxy_buffering off;
@@ -122,7 +122,7 @@ data:
 
       # public http2 port
       server {
-        listen 8091 backlog=4096 http2;
+        listen 8091 backlog={{ .Values.nginx.backlog }} http2;
         server_name public_http2;
 
         proxy_buffering off;
@@ -166,7 +166,7 @@ data:
 
       # corp port http2
       server {
-        listen 8092 backlog=4096 http2;
+        listen 8092 backlog={{ .Values.nginx.backlog }} http2;
         server_name corp_http2;
 
         proxy_buffering off;
diff --git a/Engine/Source/Programs/UnrealCloudDDC/Helm/UnrealCloudDDC/templates/deployment.yaml b/Engine/Source/Programs/UnrealCloudDDC/Helm/UnrealCloudDDC/templates/deployment.yaml
--- a/Engine/Source/Programs/UnrealCloudDDC/Helm/UnrealCloudDDC/templates/deployment.yaml
+++ b/Engine/Source/Programs/UnrealCloudDDC/Helm/UnrealCloudDDC/templates/deployment.yaml
@@ -85,6 +85,9 @@ spec:
         {{- if .Values.env }}
         {{- include "common.tplvalues.render" ( dict "value" .Values.env "context" $ ) | nindent 8 }}
         {{- end }}
+        {{- if .Values.securityContext }}
+        securityContext: {{- include "common.tplvalues.render" (dict "value" .Values.securityContext "context" $) | nindent 8 }}
+        {{- end }}
         ports:
         - name: http
           containerPort: 80
diff --git a/Engine/Source/Programs/UnrealCloudDDC/Helm/UnrealCloudDDC/values.yaml b/Engine/Source/Programs/UnrealCloudDDC/Helm/UnrealCloudDDC/values.yaml
--- a/Engine/Source/Programs/UnrealCloudDDC/Helm/UnrealCloudDDC/values.yaml
+++ b/Engine/Source/Programs/UnrealCloudDDC/Helm/UnrealCloudDDC/values.yaml
@@ -325,6 +325,12 @@ nginx:
   ## Control the number of worker connections used for nginx
   workerConnections: 1024
 
+  ## Number of connections to keep alive between nginx and Jupiter
+  keepAliveConnections: 512
+
+  ## The length of the nginx backlog, bump this if you increase `net.core.somaxconn` beyond 4096
+  backlog: 4096
+
   ## Liveness and readiness probes, defaults to corresponding endpoints
   ##
   livenessProbe:
