Unreal Cloud DDC - Corrected some setting references that were incorrect in the documentation

diff --git a/Engine/Source/Programs/UnrealCloudDDC/README.md b/Engine/Source/Programs/UnrealCloudDDC/README.md
--- a/Engine/Source/Programs/UnrealCloudDDC/README.md
+++ b/Engine/Source/Programs/UnrealCloudDDC/README.md
@@ -30,6 +30,7 @@ Unreal Cloud DDC can signficantly help teams speed up their cook processes in th
   - [Private port](#private-port)
   - [Internal Port](#internal-port)
 - [Common operations](#common-operations)
+  - [Running a local cook against a local instance](#running-a-local-cook-against-a-local-instance)
   - [Add new region](#add-new-region)
   - [Blob replication setup](#blob-replication-setup)
 
@@ -139,12 +140,13 @@ auth:
     defaultScheme: Bearer
     schemes:
       Bearer: 
-        implementation: "Jwt"
+        implementation: "JWTBearer"
         jwtAudience: "api://unreal"
         jwtAuthority: "<url-to-your-idp>
 ```
 We recommend naming your scheme `Bearer` if its your first and only scheme. You can use multiple schemes to connect against multiple IdPs, this is mostly useful during a migration.
 
+The implementation field is usually `JWTBearer` but we do offer a `Okta` if you are using Okta with custom auth servers, for Okta using the org auth server you will need to use `JWTBearer` as well.
 
 ## Namespace access
 
@@ -182,7 +184,7 @@ auth:
         - AdminAction
 
 namespace:
-  policy:
+  policies:
     example-namespace:
       acls:
       - actions: 
