Unreal Cloud DDC - Prepared version 0.5.2

diff --git a/Engine/Source/Programs/UnrealCloudDDC/Helm/UnrealCloudDDC/version.yaml b/Engine/Source/Programs/UnrealCloudDDC/Helm/UnrealCloudDDC/version.yaml
--- a/Engine/Source/Programs/UnrealCloudDDC/Helm/UnrealCloudDDC/version.yaml
+++ b/Engine/Source/Programs/UnrealCloudDDC/Helm/UnrealCloudDDC/version.yaml
@@ -1 +1 @@
-version: 0.5.1
\ No newline at end of file
+version: 0.5.2
\ No newline at end of file
diff --git a/Engine/Source/Programs/UnrealCloudDDC/changelog.md b/Engine/Source/Programs/UnrealCloudDDC/changelog.md
--- a/Engine/Source/Programs/UnrealCloudDDC/changelog.md
+++ b/Engine/Source/Programs/UnrealCloudDDC/changelog.md
@@ -1,5 +1,9 @@
 # Unreleased
 
+# 0.5.2
+* Fixed issue in helm charts defaulting to a incorrect docker registry path for the worker deployment.
+* Added ability to configure and bumped number of keepalive connections in nginx, resolves issues during large spikes of traffic.
+
 # 0.5.1
 * Fixed issue in helm charts defaulting to a incorrect docker registry path.
 
diff --git a/Engine/Source/Programs/UnrealCloudDDC/version.yaml b/Engine/Source/Programs/UnrealCloudDDC/version.yaml
--- a/Engine/Source/Programs/UnrealCloudDDC/version.yaml
+++ b/Engine/Source/Programs/UnrealCloudDDC/version.yaml
@@ -1 +1 @@
-version: 0.5.1
\ No newline at end of file
+version: 0.5.2
\ No newline at end of file
