Unreal Cloud DDC - Increasing the number of keep alive connections in nginx as its still only a fraction of the allowed worker connections, should reduce churn on connections and help with issues when we try to open to lots of connections at the same time.

diff --git a/Engine/Source/Programs/UnrealCloudDDC/Helm/UnrealCloudDDC/values.yaml b/Engine/Source/Programs/UnrealCloudDDC/Helm/UnrealCloudDDC/values.yaml
--- a/Engine/Source/Programs/UnrealCloudDDC/Helm/UnrealCloudDDC/values.yaml
+++ b/Engine/Source/Programs/UnrealCloudDDC/Helm/UnrealCloudDDC/values.yaml
@@ -326,7 +326,7 @@ nginx:
   workerConnections: 1024
 
   ## Number of connections to keep alive between nginx and Jupiter
-  keepAliveConnections: 512
+  keepAliveConnections: 2048
 
   ## The length of the nginx backlog, bump this if you increase `net.core.somaxconn` beyond 4096
   backlog: 4096
