Unreal Cloud DDC - Fixed whitespace formatting issue in configmap yaml for new bucket stats option.

diff --git a/Engine/Source/Programs/UnrealCloudDDC/Helm/UnrealCloudDDC/templates/configmap.yaml b/Engine/Source/Programs/UnrealCloudDDC/Helm/UnrealCloudDDC/templates/configmap.yaml
--- a/Engine/Source/Programs/UnrealCloudDDC/Helm/UnrealCloudDDC/templates/configmap.yaml
+++ b/Engine/Source/Programs/UnrealCloudDDC/Helm/UnrealCloudDDC/templates/configmap.yaml
@@ -29,17 +29,17 @@ UnrealCloudDDC:
   ContentIdStoreImplementation :  {{ default "Scylla" $config.ContentIdStoreImplementation }}
   BlobIndexImplementation :  {{ default "Scylla" $config.BlobIndexImplementation }}
 
-{{- if $config.EnableOnDemandReplication -}}
+{{- if $config.EnableOnDemandReplication }}
   EnableOnDemandReplication :  {{ $config.EnableOnDemandReplication }}
-{{- end -}}
+{{- end }}
 
-{{- if $config.EnableLastAccessTracking -}}
+{{- if $config.EnableLastAccessTracking }}
   EnableLastAccessTracking: {{ $config.EnableLastAccessTracking }}
-{{- end -}}
+{{- end }}
 
-{{- if $config.EnableBucketStatsTracking -}}
+{{- if $config.EnableBucketStatsTracking }}
   EnableBucketStatsTracking: {{ $config.EnableBucketStatsTracking }}
-{{- end -}}
+{{- end }}
 
 {{ if $config.StorageImplementations }}
   StorageImplementations:
