Unreal Cloud DDC - Fixed threading issues with the metrics service calculations and improved behavior when no blobs were found (making it clearer that data is default value and not valid).

diff --git a/Engine/Source/Programs/UnrealCloudDDC/Jupiter/Implementation/Blob/ScyllaBlobIndex.cs b/Engine/Source/Programs/UnrealCloudDDC/Jupiter/Implementation/Blob/ScyllaBlobIndex.cs
--- a/Engine/Source/Programs/UnrealCloudDDC/Jupiter/Implementation/Blob/ScyllaBlobIndex.cs
+++ b/Engine/Source/Programs/UnrealCloudDDC/Jupiter/Implementation/Blob/ScyllaBlobIndex.cs
@@ -4,6 +4,7 @@ using System;
 using System.Collections.Generic;
 using System.Diagnostics;
 using System.Linq;
+using System.Threading;
 using System.Threading.Tasks;
 using Cassandra;
 using Cassandra.Mapping;
@@ -370,10 +371,11 @@ public class ScyllaBlobIndex : IBlobIndex
 
 		int totalCountOfRefs = 0;
 		int totalCountOfBlobs = 0;
-		int smallestBlobFound = int.MaxValue;
-		int largestBlobFound = 0;
 		long totalSizeOfBlobs = 0;
 
+		List<int> smallestBlobPerPrefix = new List<int>();
+		List<int> largestBlobPerPrefix = new List<int>();
+
 		string[] hashPrefixes = new string[65536];
 		int i = 0;
 		for (int a = 0; a <= byte.MaxValue; a++)
@@ -398,7 +400,7 @@ public class ScyllaBlobIndex : IBlobIndex
 				{
 					int countOfRefs = (int)(long)row["system.count(reference_id)"];
 
-					totalCountOfRefs += countOfRefs;
+					Interlocked.Add(ref totalCountOfRefs, countOfRefs);
 				}
 			}, token);
 
@@ -430,16 +432,19 @@ public class ScyllaBlobIndex : IBlobIndex
 					int largestBlob = (int)(long)row["system.max(size)"];
 					long sumSizeOfBlobs = (long)row["system.sum(size)"];
 
-					totalCountOfBlobs += countOfBlobs;
-					smallestBlobFound = Math.Min(smallestBlob, smallestBlobFound);
-					largestBlobFound = Math.Max(largestBlob, largestBlobFound);
-					totalSizeOfBlobs += sumSizeOfBlobs;
+					Interlocked.Add(ref totalCountOfBlobs, countOfBlobs);
+					Interlocked.Add(ref totalSizeOfBlobs, sumSizeOfBlobs);
+
+					smallestBlobPerPrefix.Add(smallestBlob);
+					largestBlobPerPrefix.Add(largestBlob);
 				}
 			}, token);
 
 			await Task.WhenAll(calcRefStats, calcBlobStats);
 		});
-
+		int smallestBlobFound = smallestBlobPerPrefix.Any() ? smallestBlobPerPrefix.Min() : 0;
+		int largestBlobFound = largestBlobPerPrefix.Any() ? largestBlobPerPrefix.Max() : 0;
+		
 		return new BucketStats
 		{
 			Namespace = ns,
@@ -449,7 +454,9 @@ public class ScyllaBlobIndex : IBlobIndex
 			SmallestBlobFound = smallestBlobFound,
 			LargestBlob = largestBlobFound,
 			TotalSize = totalSizeOfBlobs,
-			AvgSize = totalSizeOfBlobs / (double)totalCountOfBlobs
+#pragma warning disable CA1508
+			AvgSize = totalCountOfBlobs == 0 ? 0 : (totalSizeOfBlobs / (double)totalCountOfBlobs)
+#pragma warning restore CA1508
 		};
 	}
 }
diff --git a/Engine/Source/Programs/UnrealCloudDDC/Jupiter/Implementation/MetricsService.cs b/Engine/Source/Programs/UnrealCloudDDC/Jupiter/Implementation/MetricsService.cs
--- a/Engine/Source/Programs/UnrealCloudDDC/Jupiter/Implementation/MetricsService.cs
+++ b/Engine/Source/Programs/UnrealCloudDDC/Jupiter/Implementation/MetricsService.cs
@@ -129,7 +129,7 @@ namespace Jupiter.Implementation
 			_blobSizeMinGauge = meter.CreateGauge<long>("blobstats.bucket_size.min");
 			_blobSizeMaxGauge = meter.CreateGauge<long>("blobstats.bucket_size.max");
 			_blobSizeCountGauge = meter.CreateGauge<long>("blobstats.bucket_size.count");
-			_blobSizeTotalGauge = meter.CreateGauge<long>("blobstats.bucket_size.sum,=");
+			_blobSizeTotalGauge = meter.CreateGauge<long>("blobstats.bucket_size.sum");
 			_refsInBucketGauge = meter.CreateGauge<long>("blobstats.refs_in_bucket");
 		}
 
