Unreal Cloud DDC - Added option to tweak the TCP backlog size in kestrel, will experiment with this to see if that resolves "Resource temporarily unvailable" errors from nginx

diff --git a/Engine/Source/Programs/UnrealCloudDDC/Jupiter.Common/BaseProgram.cs b/Engine/Source/Programs/UnrealCloudDDC/Jupiter.Common/BaseProgram.cs
--- a/Engine/Source/Programs/UnrealCloudDDC/Jupiter.Common/BaseProgram.cs
+++ b/Engine/Source/Programs/UnrealCloudDDC/Jupiter.Common/BaseProgram.cs
@@ -108,12 +108,22 @@ namespace Jupiter
 						Log4net.Appender.Serilog.Configuration.Configure();
 					});*/
 					// remove the server header from kestrel
+
+					JupiterSettings settings = new JupiterSettings();
+					Configuration.GetSection("Jupiter").Bind(settings);
+
+					if (settings.PendingConnectionMax.HasValue)
+					{
+						webBuilder.UseSockets(options =>
+						{
+							options.Backlog = settings.PendingConnectionMax.Value;
+						});
+					}
+
 					webBuilder.ConfigureKestrel(options =>
 					{
 						options.AddServerHeader = false;
 
-						JupiterSettings settings = new JupiterSettings();
-						Configuration.GetSection("Jupiter").Bind(settings);
 						string socketsRoot = settings.DomainSocketsRoot;
 
 						if (settings.UseDomainSockets)
diff --git a/Engine/Source/Programs/UnrealCloudDDC/Jupiter.Common/BaseStartup.cs b/Engine/Source/Programs/UnrealCloudDDC/Jupiter.Common/BaseStartup.cs
--- a/Engine/Source/Programs/UnrealCloudDDC/Jupiter.Common/BaseStartup.cs
+++ b/Engine/Source/Programs/UnrealCloudDDC/Jupiter.Common/BaseStartup.cs
@@ -724,5 +724,10 @@ namespace Jupiter
 		/// Assumes that any local connection should have full access, this is used only for tests
 		/// </summary>
 		public bool AssumeLocalConnectionsHasFullAccess { get; set; }
+
+		/// <summary>
+		/// Set this to increase the max number of TCP connections that are allowed to be queued.
+		/// </summary>
+		public int? PendingConnectionMax { get; set; }
 	}
 }
