Unreal Cloud DDC - Fixed issue when a blob was not present during finalize attempting to track which blobs were in use.

diff --git a/Engine/Source/Programs/UnrealCloudDDC/Jupiter/Implementation/RefService.cs b/Engine/Source/Programs/UnrealCloudDDC/Jupiter/Implementation/RefService.cs
--- a/Engine/Source/Programs/UnrealCloudDDC/Jupiter/Implementation/RefService.cs
+++ b/Engine/Source/Programs/UnrealCloudDDC/Jupiter/Implementation/RefService.cs
@@ -222,8 +222,15 @@ namespace Jupiter.Implementation
 						{
 							addToBucketTasks.Add(Task.Run( async () =>
 							{
-								BlobContents result = await _blobService.GetObjectAsync(ns, blobId);
-								await _blobIndex.AddBlobToBucketListAsync(ns, bucket, key, blobId, result.Length);
+								// if a blob is missing its not a error, the finalize will report this as missing and it will be uploaded and finalize ran again
+								try
+								{
+									BlobContents result = await _blobService.GetObjectAsync(ns, blobId);
+									await _blobIndex.AddBlobToBucketListAsync(ns, bucket, key, blobId, result.Length);
+								}
+								catch (BlobNotFoundException)
+								{
+								}
 							}));
 						}
 
