Unreal Cloud DDC - Added explict tracking of surpressed exception

diff --git a/Engine/Source/Programs/UnrealCloudDDC/Jupiter.Common/PollingService.cs b/Engine/Source/Programs/UnrealCloudDDC/Jupiter.Common/PollingService.cs
--- a/Engine/Source/Programs/UnrealCloudDDC/Jupiter.Common/PollingService.cs
+++ b/Engine/Source/Programs/UnrealCloudDDC/Jupiter.Common/PollingService.cs
@@ -6,6 +6,7 @@ using System.Threading;
 using System.Threading.Tasks;
 using Microsoft.Extensions.Hosting;
 using Microsoft.Extensions.Logging;
+using OpenTelemetry.Trace;
 
 namespace Jupiter
 {
@@ -134,6 +135,8 @@ namespace Jupiter
 			catch (Exception e)
 			{
 				logger.LogError(e, "{Service} Exception in polling service", serviceName);
+				Tracer.CurrentSpan.SetStatus(Status.Error);
+				Tracer.CurrentSpan.RecordException(e);
 			}
 			finally
 			{
diff --git a/Engine/Source/Programs/UnrealCloudDDC/Jupiter/Controllers/ReferencesController.cs b/Engine/Source/Programs/UnrealCloudDDC/Jupiter/Controllers/ReferencesController.cs
--- a/Engine/Source/Programs/UnrealCloudDDC/Jupiter/Controllers/ReferencesController.cs
+++ b/Engine/Source/Programs/UnrealCloudDDC/Jupiter/Controllers/ReferencesController.cs
@@ -370,6 +370,8 @@ namespace Jupiter.Controllers
 								}
 								catch (Exception ex)
 								{
+									Tracer.CurrentSpan.SetStatus(Status.Error);
+									Tracer.CurrentSpan.RecordException(ex);
 									_logger.LogError(ex, "Unknown exception encountered while writing body for jupiter inlined payload.");
 									throw;
 								}
@@ -882,6 +884,8 @@ namespace Jupiter.Controllers
 				}
 				catch (Exception e)
 				{
+					Tracer.CurrentSpan.SetStatus(Status.Error);
+					Tracer.CurrentSpan.RecordException(e);
 					return ToErrorResult(e);
 				}
 			}
diff --git a/Engine/Source/Programs/UnrealCloudDDC/Jupiter/Implementation/LastAccess/LastAccessServiceReferences.cs b/Engine/Source/Programs/UnrealCloudDDC/Jupiter/Implementation/LastAccess/LastAccessServiceReferences.cs
--- a/Engine/Source/Programs/UnrealCloudDDC/Jupiter/Implementation/LastAccess/LastAccessServiceReferences.cs
+++ b/Engine/Source/Programs/UnrealCloudDDC/Jupiter/Implementation/LastAccess/LastAccessServiceReferences.cs
@@ -68,6 +68,8 @@ namespace Jupiter.Implementation
 			catch (Exception e)
 			{
 				_logger.LogError(e, "Exception thrown while processing last access records");
+				Tracer.CurrentSpan.SetStatus(Status.Error);
+				Tracer.CurrentSpan.RecordException(e);
 			}
 		}
 
