Unreal Cloud DDC - Fixed handling of non finalized objects in batch endpoint, won't show up as an error anymore.

diff --git a/Engine/Source/Programs/UnrealCloudDDC/Jupiter/Controllers/ReferencesController.cs b/Engine/Source/Programs/UnrealCloudDDC/Jupiter/Controllers/ReferencesController.cs
--- a/Engine/Source/Programs/UnrealCloudDDC/Jupiter/Controllers/ReferencesController.cs
+++ b/Engine/Source/Programs/UnrealCloudDDC/Jupiter/Controllers/ReferencesController.cs
@@ -860,7 +860,7 @@ namespace Jupiter.Controllers
 
 					if (!objectRecord.IsFinalized)
 					{
-						throw new Exception("Object is not finalized");
+						return ToErrorResult("object not finalized", HttpStatusCode.BadRequest);
 					}
 
 					if (blob == null)
@@ -1011,6 +1011,16 @@ namespace Jupiter.Controllers
 			return (writer.ToObject(), statusCode);
 		}
 
+		private static (CbObject, HttpStatusCode) ToErrorResult(string message, HttpStatusCode statusCode = HttpStatusCode.InternalServerError)
+		{
+			CbWriter writer = new CbWriter();
+			writer.BeginObject();
+			writer.WriteString("title", message);
+			writer.WriteInteger("status", (int)statusCode);
+			writer.EndObject();
+			return (writer.ToObject(), statusCode);
+		}
+
 		/// <summary>
 		/// Drop all refs records in the namespace
 		/// </summary>
