Unreal Cloud DDC - Increased timeout nginx gives Cloud DDC to reply, resolves issues with nginx aborting uploads of really large payloads (primarilly a issue for oplogs)

diff --git a/Engine/Source/Programs/UnrealCloudDDC/Helm/UnrealCloudDDC/templates/configmap-nginx.yaml b/Engine/Source/Programs/UnrealCloudDDC/Helm/UnrealCloudDDC/templates/configmap-nginx.yaml
--- a/Engine/Source/Programs/UnrealCloudDDC/Helm/UnrealCloudDDC/templates/configmap-nginx.yaml
+++ b/Engine/Source/Programs/UnrealCloudDDC/Helm/UnrealCloudDDC/templates/configmap-nginx.yaml
@@ -26,6 +26,9 @@ data:
       sendfile on;
       tcp_nopush on;
 
+      # set communication timeout to 5 minutes, transfering really large files takes a while
+      proxy_read_timeout 300s
+
       upstream jupiter_backend {
         zone upstreams 64K;
         server {{ if .Values.nginx.useDomainSockets -}} unix:/nginx/jupiter-http.sock {{- else -}} localhost:80 {{- end }} max_fails=1 fail_timeout=2s;
