Unreal Cloud DDC - Added missing ; after proxy_timeout directive

diff --git a/Engine/Source/Programs/UnrealCloudDDC/Helm/UnrealCloudDDC/templates/configmap-nginx.yaml b/Engine/Source/Programs/UnrealCloudDDC/Helm/UnrealCloudDDC/templates/configmap-nginx.yaml
--- a/Engine/Source/Programs/UnrealCloudDDC/Helm/UnrealCloudDDC/templates/configmap-nginx.yaml
+++ b/Engine/Source/Programs/UnrealCloudDDC/Helm/UnrealCloudDDC/templates/configmap-nginx.yaml
@@ -27,7 +27,7 @@ data:
       tcp_nopush on;
 
       # set communication timeout to 5 minutes, transfering really large files takes a while
-      proxy_read_timeout 300s
+      proxy_read_timeout 300s;
 
       upstream jupiter_backend {
         zone upstreams 64K;
