Unreal Cloud DDC - Added check that Default Keyspace has been specified as it needs to be set if the connection string has been overriden.

diff --git a/Engine/Source/Programs/UnrealCloudDDC/Jupiter/Jupiter.Startup.cs b/Engine/Source/Programs/UnrealCloudDDC/Jupiter/Jupiter.Startup.cs
--- a/Engine/Source/Programs/UnrealCloudDDC/Jupiter/Jupiter.Startup.cs
+++ b/Engine/Source/Programs/UnrealCloudDDC/Jupiter/Jupiter.Startup.cs
@@ -297,6 +297,12 @@ namespace Jupiter
 				throw new Exception("Connection string has to be specified");
 			}
 
+			CassandraConnectionStringBuilder cassandraConnectionStringBuilder = new CassandraConnectionStringBuilder(connectionString);
+			if (string.IsNullOrEmpty(cassandraConnectionStringBuilder.DefaultKeyspace))
+			{
+				throw new Exception("Default Keyspace has to be specified");
+			}
+
 			// Configure the builder with your cluster's contact points
 			Builder clusterBuilder = Cluster.Builder()
 				.WithConnectionString(connectionString)
