Unreal Cloud DDC - Added option to disable chunked encoding for S3 uploads - required for GCS (Google Storage)

diff --git a/Engine/Source/Programs/UnrealCloudDDC/Helm/UnrealCloudDDC/templates/configmap.yaml b/Engine/Source/Programs/UnrealCloudDDC/Helm/UnrealCloudDDC/templates/configmap.yaml
--- a/Engine/Source/Programs/UnrealCloudDDC/Helm/UnrealCloudDDC/templates/configmap.yaml
+++ b/Engine/Source/Programs/UnrealCloudDDC/Helm/UnrealCloudDDC/templates/configmap.yaml
@@ -190,5 +190,10 @@ S3:
   {{- if $s3.UseMultiPartUpload }}
   UseMultiPartUpload: {{ $s3.UseMultiPartUpload }}
   {{- end }}
+
+  {{- if $s3.UseChunkEncoding }}
+  UseChunkEncoding: {{ $s3.UseChunkEncoding }}
+  {{- end }}
+  
 {{ end }}
 
diff --git a/Engine/Source/Programs/UnrealCloudDDC/Jupiter/Implementation/Blob/AmazonS3Store.cs b/Engine/Source/Programs/UnrealCloudDDC/Jupiter/Implementation/Blob/AmazonS3Store.cs
--- a/Engine/Source/Programs/UnrealCloudDDC/Jupiter/Implementation/Blob/AmazonS3Store.cs
+++ b/Engine/Source/Programs/UnrealCloudDDC/Jupiter/Implementation/Blob/AmazonS3Store.cs
@@ -266,7 +266,8 @@ namespace Jupiter.Implementation
 				{
 					BucketName = _bucketName,
 					Key = path,
-					InputStream = stream
+					InputStream = stream,
+					UseChunkEncoding = _settings.CurrentValue.UseChunkEncoding
 				};
 
 				try
diff --git a/Engine/Source/Programs/UnrealCloudDDC/Jupiter/Settings.cs b/Engine/Source/Programs/UnrealCloudDDC/Jupiter/Settings.cs
--- a/Engine/Source/Programs/UnrealCloudDDC/Jupiter/Settings.cs
+++ b/Engine/Source/Programs/UnrealCloudDDC/Jupiter/Settings.cs
@@ -183,6 +183,11 @@ namespace Jupiter
 
 		public bool? UseArnRegion { get; set; } = null;
 		public bool UseMultiPartUpload { get; set; } = false;
+
+		/// <summary>
+		/// Allows you to override S3 behavior with chunk encoding, this needs to be set to false for uploads against GCS
+		/// </summary>
+		public bool UseChunkEncoding { get; set; } = true;
 	}
 
 	public class GCSettings
