Unreal Cloud DDC - Prepared version 0.6.0

diff --git a/Engine/Source/Programs/UnrealCloudDDC/Helm/UnrealCloudDDC/version.yaml b/Engine/Source/Programs/UnrealCloudDDC/Helm/UnrealCloudDDC/version.yaml
--- a/Engine/Source/Programs/UnrealCloudDDC/Helm/UnrealCloudDDC/version.yaml
+++ b/Engine/Source/Programs/UnrealCloudDDC/Helm/UnrealCloudDDC/version.yaml
@@ -1 +1 @@
-version: 0.5.2
\ No newline at end of file
+version: 0.6.0
\ No newline at end of file
diff --git a/Engine/Source/Programs/UnrealCloudDDC/changelog.md b/Engine/Source/Programs/UnrealCloudDDC/changelog.md
--- a/Engine/Source/Programs/UnrealCloudDDC/changelog.md
+++ b/Engine/Source/Programs/UnrealCloudDDC/changelog.md
@@ -1,5 +1,13 @@
 # Unreleased
 
+# 0.6.0
+* Added option to track per bucket stats `EnableBucketStatsTracking`, this is still WIP.
+* Added option to tweak nginx keep alive connections and increased it, only applies if using the nginx proxy.
+* Added option to control TCP backlog in kestrel.
+* Improved error message when no keyspace is set to Scylla
+* Increased proxy timeout when using nginx, allowing for long operations (upload of oplogs) to not timeout.
+* Added option to disable chunk encoding `S3.UseChunkEncoding` which needs to be set when using GCS.
+
 # 0.5.2
 * Fixed issue in helm charts defaulting to a incorrect docker registry path for the worker deployment.
 * Added ability to configure and bumped number of keepalive connections in nginx, resolves issues during large spikes of traffic.
diff --git a/Engine/Source/Programs/UnrealCloudDDC/version.yaml b/Engine/Source/Programs/UnrealCloudDDC/version.yaml
--- a/Engine/Source/Programs/UnrealCloudDDC/version.yaml
+++ b/Engine/Source/Programs/UnrealCloudDDC/version.yaml
@@ -1 +1 @@
-version: 0.5.2
\ No newline at end of file
+version: 0.6.0
\ No newline at end of file
