Missing files.

diff --git a/Engine/Source/Programs/UnrealCloudDDC/Jupiter.Common/Implementation/BlobId.cs b/Engine/Source/Programs/UnrealCloudDDC/Jupiter.Common/Implementation/BlobId.cs
--- a/Engine/Source/Programs/UnrealCloudDDC/Jupiter.Common/Implementation/BlobId.cs
+++ b/Engine/Source/Programs/UnrealCloudDDC/Jupiter.Common/Implementation/BlobId.cs
@@ -221,6 +221,6 @@ namespace Jupiter.Implementation
 		public override void Write(CbWriter writer, BlobId value) => writer.WriteBinaryAttachmentValue(new IoHash(value.HashData));
 
 		/// <inheritdoc/>
-		public override void WriteNamed(CbWriter writer, Utf8String name, BlobId value) => writer.WriteBinaryAttachment(name, new IoHash(value.HashData));
+		public override void WriteNamed(CbWriter writer, CbFieldName name, BlobId value) => writer.WriteBinaryAttachment(name, new IoHash(value.HashData));
 	}
 }
diff --git a/Engine/Source/Programs/UnrealCloudDDC/Jupiter.Common/Implementation/BucketId.cs b/Engine/Source/Programs/UnrealCloudDDC/Jupiter.Common/Implementation/BucketId.cs
--- a/Engine/Source/Programs/UnrealCloudDDC/Jupiter.Common/Implementation/BucketId.cs
+++ b/Engine/Source/Programs/UnrealCloudDDC/Jupiter.Common/Implementation/BucketId.cs
@@ -83,6 +83,6 @@ namespace EpicGames.Horde.Storage
 		public override void Write(CbWriter writer, BucketId value) => writer.WriteStringValue(value.ToString());
 
 		/// <inheritdoc/>
-		public override void WriteNamed(CbWriter writer, Utf8String name, BucketId value) => writer.WriteString(name, value.ToString());
+		public override void WriteNamed(CbWriter writer, CbFieldName name, BucketId value) => writer.WriteString(name, value.ToString());
 	}
 }
diff --git a/Engine/Source/Programs/UnrealCloudDDC/Jupiter.Common/Implementation/ContentHash.cs b/Engine/Source/Programs/UnrealCloudDDC/Jupiter.Common/Implementation/ContentHash.cs
--- a/Engine/Source/Programs/UnrealCloudDDC/Jupiter.Common/Implementation/ContentHash.cs
+++ b/Engine/Source/Programs/UnrealCloudDDC/Jupiter.Common/Implementation/ContentHash.cs
@@ -118,7 +118,7 @@ namespace Jupiter.Implementation
 		public override void Write(CbWriter writer, ContentHash value) => writer.WriteHashValue(new IoHash(value.HashData));
 
 		/// <inheritdoc/>
-		public override void WriteNamed(CbWriter writer, Utf8String name, ContentHash value) => writer.WriteHash(name, new IoHash(value.HashData));
+		public override void WriteNamed(CbWriter writer, CbFieldName name, ContentHash value) => writer.WriteHash(name, new IoHash(value.HashData));
 	}
 
 	public class ContentHashJsonConverter : JsonConverter<ContentHash>
diff --git a/Engine/Source/Programs/UnrealCloudDDC/Jupiter.Common/Implementation/RefId.cs b/Engine/Source/Programs/UnrealCloudDDC/Jupiter.Common/Implementation/RefId.cs
--- a/Engine/Source/Programs/UnrealCloudDDC/Jupiter.Common/Implementation/RefId.cs
+++ b/Engine/Source/Programs/UnrealCloudDDC/Jupiter.Common/Implementation/RefId.cs
@@ -137,7 +137,7 @@ namespace Jupiter.Implementation
 		public override void Write(CbWriter writer, RefId value) => writer.WriteStringValue(value.ToString());
 
 		/// <inheritdoc/>
-		public override void WriteNamed(CbWriter writer, Utf8String name, RefId value) => writer.WriteString(name, value.ToString());
+		public override void WriteNamed(CbWriter writer, CbFieldName name, RefId value) => writer.WriteString(name, value.ToString());
 	}
 }
 
