Unreal Cloud DDC - Added ref store consistency check that will find an fix refs that lack last access data, resulting in us missing it when GCing.

diff --git a/Engine/Source/Programs/UnrealCloudDDC/Jupiter/Implementation/Objects/CacheReferencesStore.cs b/Engine/Source/Programs/UnrealCloudDDC/Jupiter/Implementation/Objects/CacheReferencesStore.cs
--- a/Engine/Source/Programs/UnrealCloudDDC/Jupiter/Implementation/Objects/CacheReferencesStore.cs
+++ b/Engine/Source/Programs/UnrealCloudDDC/Jupiter/Implementation/Objects/CacheReferencesStore.cs
@@ -64,6 +64,26 @@ namespace Jupiter.Implementation
 			await Task.WhenAll(cacheFinalize, upstreamFinalize);
 		}
 
+		public async Task<DateTime?> GetLastAccessTimeAsync(NamespaceId ns, BucketId bucket, RefId key)
+		{
+			try
+			{
+				DateTime? lastAccessTimeCache = await _mongoReferenceStore.GetLastAccessTimeAsync(ns, bucket, key);
+				return lastAccessTimeCache;
+			}
+			catch (RefNotFoundException)
+			{
+				// not cached, we check the upstream for it
+			}
+			DateTime? lastAccessTime = await _upstreamReferenceStore.GetLastAccessTimeAsync(ns, bucket, key);
+			if (lastAccessTime.HasValue)
+			{
+				await _mongoReferenceStore.UpdateLastAccessTimeAsync(ns, bucket, key, lastAccessTime.Value);
+			}
+
+			return lastAccessTime;
+		}
+
 		public async Task UpdateLastAccessTimeAsync(NamespaceId ns, BucketId bucket, RefId key, DateTime newLastAccessTime)
 		{
 			Task cacheUpdateLastAccess = _mongoReferenceStore.UpdateLastAccessTimeAsync(ns, bucket, key, newLastAccessTime);
@@ -190,6 +210,19 @@ namespace Jupiter.Implementation
 			}
 		}
 
+		public async Task<DateTime?> GetLastAccessTimeAsync(NamespaceId ns, BucketId bucket, RefId key)
+		{
+			try
+			{
+				RefRecord foo = await GetAsync(ns, bucket, key, IReferencesStore.FieldFlags.None, IReferencesStore.OperationFlags.None);
+				return foo.LastAccess;
+			}
+			catch (RefNotFoundException)
+			{
+				return null;
+			}
+		}
+
 		public async Task UpdateLastAccessTimeAsync(NamespaceId ns, BucketId bucket, RefId key, DateTime newLastAccessTime)
 		{
 			// there is no endpoint to update last access time externally, but fetching a object will in turn update its last access time, though to a different time then newLastAccessTime
diff --git a/Engine/Source/Programs/UnrealCloudDDC/Jupiter/Implementation/Objects/IReferencesStore.cs b/Engine/Source/Programs/UnrealCloudDDC/Jupiter/Implementation/Objects/IReferencesStore.cs
--- a/Engine/Source/Programs/UnrealCloudDDC/Jupiter/Implementation/Objects/IReferencesStore.cs
+++ b/Engine/Source/Programs/UnrealCloudDDC/Jupiter/Implementation/Objects/IReferencesStore.cs
@@ -29,6 +29,7 @@ namespace Jupiter.Implementation
 		Task PutAsync(NamespaceId ns, BucketId bucket, RefId key, BlobId blobHash, byte[] blob, bool isFinalized);
 		Task FinalizeAsync(NamespaceId ns, BucketId bucket, RefId key, BlobId blobIdentifier);
 
+		Task<DateTime?> GetLastAccessTimeAsync(NamespaceId ns, BucketId bucket, RefId key);
 		Task UpdateLastAccessTimeAsync(NamespaceId ns, BucketId bucket, RefId key, DateTime newLastAccessTime);
 		IAsyncEnumerable<(NamespaceId, BucketId, RefId, DateTime)> GetRecordsAsync();
 
diff --git a/Engine/Source/Programs/UnrealCloudDDC/Jupiter/Implementation/Objects/MemoryCachedReferencesStore.cs b/Engine/Source/Programs/UnrealCloudDDC/Jupiter/Implementation/Objects/MemoryCachedReferencesStore.cs
--- a/Engine/Source/Programs/UnrealCloudDDC/Jupiter/Implementation/Objects/MemoryCachedReferencesStore.cs
+++ b/Engine/Source/Programs/UnrealCloudDDC/Jupiter/Implementation/Objects/MemoryCachedReferencesStore.cs
@@ -92,6 +92,11 @@ namespace Jupiter.Implementation.Objects
 			return _actualStore.FinalizeAsync(ns, bucket, key, blobIdentifier);
 		}
 
+		public Task<DateTime?> GetLastAccessTimeAsync(NamespaceId ns, BucketId bucket, RefId key)
+		{
+			return _actualStore.GetLastAccessTimeAsync(ns, bucket, key);
+		}
+
 		private void FinalizeCacheEntry(NamespaceId ns, BucketId bucket, RefId key)
 		{
 			MemoryCache cache = GetCacheForNamespace(ns);
diff --git a/Engine/Source/Programs/UnrealCloudDDC/Jupiter/Implementation/Objects/MemoryReferencesStore.cs b/Engine/Source/Programs/UnrealCloudDDC/Jupiter/Implementation/Objects/MemoryReferencesStore.cs
--- a/Engine/Source/Programs/UnrealCloudDDC/Jupiter/Implementation/Objects/MemoryReferencesStore.cs
+++ b/Engine/Source/Programs/UnrealCloudDDC/Jupiter/Implementation/Objects/MemoryReferencesStore.cs
@@ -54,6 +54,17 @@ namespace Jupiter.Implementation
 			return Task.CompletedTask;
 		}
 
+		public Task<DateTime?> GetLastAccessTimeAsync(NamespaceId ns, BucketId bucket, RefId key)
+		{
+			DateTime? lastAccessTime = null;
+			if (_objects.TryGetValue(BuildKey(ns, bucket, key), out MemoryStoreObject? o))
+			{
+				lastAccessTime = o.LastAccessTime;
+			}
+
+			return Task.FromResult(lastAccessTime);
+		}
+
 		public Task UpdateLastAccessTimeAsync(NamespaceId ns, BucketId bucket, RefId key, DateTime lastAccessTime)
 		{
 			if (!_objects.TryGetValue(BuildKey(ns, bucket, key), out MemoryStoreObject? o))
diff --git a/Engine/Source/Programs/UnrealCloudDDC/Jupiter/Implementation/Objects/MongoReferencesStore.cs b/Engine/Source/Programs/UnrealCloudDDC/Jupiter/Implementation/Objects/MongoReferencesStore.cs
--- a/Engine/Source/Programs/UnrealCloudDDC/Jupiter/Implementation/Objects/MongoReferencesStore.cs
+++ b/Engine/Source/Programs/UnrealCloudDDC/Jupiter/Implementation/Objects/MongoReferencesStore.cs
@@ -108,6 +108,15 @@ namespace Jupiter.Implementation
 			);
 		}
 
+		public async Task<DateTime?> GetLastAccessTimeAsync(NamespaceId ns, BucketId bucket, RefId key)
+		{
+			IMongoCollection<MongoReferencesModelV0> collection = GetCollection<MongoReferencesModelV0>();
+			IAsyncCursor<MongoReferencesModelV0>? cursor = await collection.FindAsync(m => m.Ns == ns.ToString() && m.Bucket == bucket.ToString() && m.Key == key.ToString());
+			MongoReferencesModelV0? model = await cursor.FirstOrDefaultAsync();
+
+			return model?.LastAccessTime;
+		}
+
 		public async Task UpdateLastAccessTimeAsync(NamespaceId ns, BucketId bucket, RefId key, DateTime newLastAccessTime)
 		{
 			IMongoCollection<MongoReferencesModelV0> collection = GetCollection<MongoReferencesModelV0>();
diff --git a/Engine/Source/Programs/UnrealCloudDDC/Jupiter/Implementation/Objects/ScyllaReferencesStore.cs b/Engine/Source/Programs/UnrealCloudDDC/Jupiter/Implementation/Objects/ScyllaReferencesStore.cs
--- a/Engine/Source/Programs/UnrealCloudDDC/Jupiter/Implementation/Objects/ScyllaReferencesStore.cs
+++ b/Engine/Source/Programs/UnrealCloudDDC/Jupiter/Implementation/Objects/ScyllaReferencesStore.cs
@@ -180,6 +180,13 @@ namespace Jupiter.Implementation
 			await _mapper.UpdateAsync<ScyllaObject>("SET is_finalized=true WHERE namespace=? AND bucket=? AND name=?", ns.ToString(), bucket.ToString(), name.ToString());
 		}
 
+		public async Task<DateTime?> GetLastAccessTimeAsync(NamespaceId ns, BucketId bucket, RefId key)
+		{
+			using TelemetrySpan scope = _tracer.BuildScyllaSpan("scylla.get_last_access_time").SetAttribute("resource.name", $"{ns}.{bucket}.{key}");
+			ScyllaObjectLastAccess? lastAccessRecord = await _mapper.SingleOrDefaultAsync<ScyllaObjectLastAccess>("WHERE namespace = ? AND bucket = ? AND name = ? {cqlOptions}", ns.ToString(), bucket.ToString(), key.ToString());
+			return lastAccessRecord?.LastAccessTime;
+		}
+
 		public async Task UpdateLastAccessTimeAsync(NamespaceId ns, BucketId bucket, RefId name, DateTime lastAccessTime)
 		{
 			using TelemetrySpan scope = _tracer.BuildScyllaSpan("scylla.update_last_access_time");
diff --git a/Engine/Source/Programs/UnrealCloudDDC/Jupiter/Jupiter.Startup.cs b/Engine/Source/Programs/UnrealCloudDDC/Jupiter/Jupiter.Startup.cs
--- a/Engine/Source/Programs/UnrealCloudDDC/Jupiter/Jupiter.Startup.cs
+++ b/Engine/Source/Programs/UnrealCloudDDC/Jupiter/Jupiter.Startup.cs
@@ -183,6 +183,9 @@ namespace Jupiter
 			services.AddSingleton<BlobIndexConsistencyCheckService>();
 			services.AddHostedService<BlobIndexConsistencyCheckService>(p => p.GetService<BlobIndexConsistencyCheckService>()!);
 
+			services.AddSingleton<RefStoreConsistencyCheckService>();
+			services.AddHostedService<RefStoreConsistencyCheckService>(p => p.GetService<RefStoreConsistencyCheckService>()!);
+
 			services.AddSingleton<MetricsService>();
 			services.AddHostedService<MetricsService>(p => p.GetService<MetricsService>()!);
 
diff --git a/Engine/Source/Programs/UnrealCloudDDC/Jupiter/Settings.cs b/Engine/Source/Programs/UnrealCloudDDC/Jupiter/Settings.cs
--- a/Engine/Source/Programs/UnrealCloudDDC/Jupiter/Settings.cs
+++ b/Engine/Source/Programs/UnrealCloudDDC/Jupiter/Settings.cs
@@ -240,6 +240,7 @@ namespace Jupiter
 	{
 		public bool EnableBlobStoreChecks { get; set; } = false;
 		public bool EnableBlobIndexChecks { get; set; } = false;
+		public bool EnableRefStoreChecks { get; set; } = false;
 		public double ConsistencyCheckPollFrequencySeconds { get; set; } = TimeSpan.FromHours(2).TotalSeconds;
 		public int BlobIndexMaxParallelOperations { get; set; } = 4;
 		public bool AllowDeletesInBlobIndex { get; set; } = false;
