Unreal Cloud DDC - Do not attempt to populate earlier blob store layers with data that is being returned as a redirect uri as this will not contain any data.

diff --git a/Engine/Source/Programs/UnrealCloudDDC/Jupiter/Implementation/Blob/BlobService.cs b/Engine/Source/Programs/UnrealCloudDDC/Jupiter/Implementation/Blob/BlobService.cs
--- a/Engine/Source/Programs/UnrealCloudDDC/Jupiter/Implementation/Blob/BlobService.cs
+++ b/Engine/Source/Programs/UnrealCloudDDC/Jupiter/Implementation/Blob/BlobService.cs
@@ -407,7 +407,8 @@ public class BlobService : IBlobService
 			throw new Exception("blobContents is null");
 		}
 
-		if (numStoreMisses >= 1)
+		// if we had a miss populate the earlier stores unless we are using a redirect uri, at which point we do not have the contents available to us
+		if (numStoreMisses >= 1 && blobContents.RedirectUri == null)
 		{
 			using TelemetrySpan _ = _tracer.StartActiveSpan("HierarchicalStore.Populate").SetAttribute("operation.name", "HierarchicalStore.Populate");
 			using ServerTimingMetricScoped? serverTimingScope = serverTiming?.CreateServerTimingMetricScope($"blob.populate", "Populating caches with blob contents");
