Unreal Cloud DDC - Fixed broken tracing as a result of the .NET 8 upgrade

diff --git a/Engine/Source/Programs/UnrealCloudDDC/Jupiter.Common/BaseStartup.cs b/Engine/Source/Programs/UnrealCloudDDC/Jupiter.Common/BaseStartup.cs
--- a/Engine/Source/Programs/UnrealCloudDDC/Jupiter.Common/BaseStartup.cs
+++ b/Engine/Source/Programs/UnrealCloudDDC/Jupiter.Common/BaseStartup.cs
@@ -259,9 +259,7 @@ namespace Jupiter
 				.AddService("UnrealCloudDDC", serviceNamespace: "Jupiter", serviceVersion: otelServiceVersion)
 				.AddEnvironmentVariableDetector();
 
-			services.AddOpenTelemetry();
-
-			services.ConfigureOpenTelemetryTracerProvider(builder =>
+			services.AddOpenTelemetry().WithTracing(builder =>
 			{
 				builder.AddHttpClientInstrumentation(options =>
 				{
@@ -302,9 +300,7 @@ namespace Jupiter
 				builder.AddOtlpExporter();
 
 				builder.AddSource("UnrealCloudDDC", "ScyllaDB");
-			});
-
-			services.ConfigureOpenTelemetryMeterProvider(builder =>
+			}).WithMetrics(builder =>
 			{
 				builder
 					.AddMeter("UnrealCloudDDC", "ScyllaDB")
@@ -313,7 +309,7 @@ namespace Jupiter
 					.AddAspNetCoreInstrumentation()
 					.AddHttpClientInstrumentation();
 			});
-			
+
 			services.Configure<OpenTelemetryLoggerOptions>(opt =>
 			{
 				opt.IncludeScopes = true;
