Unreal Cloud DDC - Further tweaks to otel tracing configuration as well as added a option to enable the console exporter for easier debugging

diff --git a/Engine/Source/Programs/UnrealCloudDDC/Jupiter.Common/BaseStartup.cs b/Engine/Source/Programs/UnrealCloudDDC/Jupiter.Common/BaseStartup.cs
--- a/Engine/Source/Programs/UnrealCloudDDC/Jupiter.Common/BaseStartup.cs
+++ b/Engine/Source/Programs/UnrealCloudDDC/Jupiter.Common/BaseStartup.cs
@@ -255,11 +255,14 @@ namespace Jupiter
 			string otelServiceName = Configuration["OTEL_SERVICE_NAME"] ?? "unreal-cloud-ddc";
 			string? otelServiceVersion = Configuration["OTEL_SERVICE_VERSION"];
 
-			ResourceBuilder appResourceBuilder = ResourceBuilder.CreateDefault()
-				.AddService("UnrealCloudDDC", serviceNamespace: "Jupiter", serviceVersion: otelServiceVersion)
-				.AddEnvironmentVariableDetector();
+			string? useConsoleExporterString = Configuration["OTEL_USE_CONSOLE_EXPORTER"];
+			_ = bool.TryParse(useConsoleExporterString, out bool useConsoleExporter);
 
-			services.AddOpenTelemetry().WithTracing(builder =>
+			services.AddOpenTelemetry().ConfigureResource(builder =>
+			{
+				builder.AddService("UnrealCloudDDC", serviceNamespace: "Jupiter", serviceVersion: otelServiceVersion)
+				.AddEnvironmentVariableDetector();
+			}).WithTracing(builder =>
 			{
 				builder.AddHttpClientInstrumentation(options =>
 				{
@@ -296,18 +299,24 @@ namespace Jupiter
 					};
 				});
 
-				builder.SetResourceBuilder(appResourceBuilder);
 				builder.AddOtlpExporter();
-
+				if (useConsoleExporter)
+				{
+					builder.AddConsoleExporter();
+				}
 				builder.AddSource("UnrealCloudDDC", "ScyllaDB");
 			}).WithMetrics(builder =>
 			{
 				builder
 					.AddMeter("UnrealCloudDDC", "ScyllaDB")
 					.AddOtlpExporter()
-					.SetResourceBuilder(appResourceBuilder)
 					.AddAspNetCoreInstrumentation()
 					.AddHttpClientInstrumentation();
+
+				if (useConsoleExporter)
+				{
+					builder.AddConsoleExporter();
+				}
 			});
 
 			services.Configure<OpenTelemetryLoggerOptions>(opt =>
