Unreal Cloud DDC - Increased nginx send timeout to match the read timeout to resolve 504 issues when fetching really large files (9GB+)

diff --git a/Engine/Source/Programs/UnrealCloudDDC/Helm/UnrealCloudDDC/templates/configmap-nginx.yaml b/Engine/Source/Programs/UnrealCloudDDC/Helm/UnrealCloudDDC/templates/configmap-nginx.yaml
--- a/Engine/Source/Programs/UnrealCloudDDC/Helm/UnrealCloudDDC/templates/configmap-nginx.yaml
+++ b/Engine/Source/Programs/UnrealCloudDDC/Helm/UnrealCloudDDC/templates/configmap-nginx.yaml
@@ -28,6 +28,7 @@ data:
 
       # set communication timeout to 5 minutes, transfering really large files takes a while
       proxy_read_timeout 300s;
+      proxy_send_timeout 300s;
 
       upstream jupiter_backend {
         zone upstreams 64K;
