Unreal Cloud DDC - Introduced disabled leader election - fixes potential issue were individual pods thought they were leaders when worker processes was being used

diff --git a/Engine/Source/Programs/UnrealCloudDDC/Helm/UnrealCloudDDC/templates/configmap-worker.yaml b/Engine/Source/Programs/UnrealCloudDDC/Helm/UnrealCloudDDC/templates/configmap-worker.yaml
--- a/Engine/Source/Programs/UnrealCloudDDC/Helm/UnrealCloudDDC/templates/configmap-worker.yaml
+++ b/Engine/Source/Programs/UnrealCloudDDC/Helm/UnrealCloudDDC/templates/configmap-worker.yaml
@@ -20,7 +20,8 @@ data:
 {{- $filesystem := $config.Filesystem | default dict -}}
 
 UnrealCloudDDC:
-  LeaderElectionImplementation: "Kubernetes"
+  # always assume you are the leader when you are the worker
+  LeaderElectionImplementation: "Static"
 
   ReferencesDbImplementation:  {{ default "Scylla" $config.ReferencesDbImplementation }}
   ReplicationLogWriterImplementation:  {{ default "Scylla" $config.ReplicationLogWriterImplementation }}
diff --git a/Engine/Source/Programs/UnrealCloudDDC/Helm/UnrealCloudDDC/templates/configmap.yaml b/Engine/Source/Programs/UnrealCloudDDC/Helm/UnrealCloudDDC/templates/configmap.yaml
--- a/Engine/Source/Programs/UnrealCloudDDC/Helm/UnrealCloudDDC/templates/configmap.yaml
+++ b/Engine/Source/Programs/UnrealCloudDDC/Helm/UnrealCloudDDC/templates/configmap.yaml
@@ -18,7 +18,7 @@ data:
 UnrealCloudDDC:
 {{ if .Values.worker.enabled }}
   # if using the workers we do not attempt to participate in the leader election as a normal deployment
-  LeaderElectionImplementation: "Static"
+  LeaderElectionImplementation: "Disabled"
 {{ else }}
   LeaderElectionImplementation: "Kubernetes"
 {{ end }}
diff --git a/Engine/Source/Programs/UnrealCloudDDC/Jupiter/Jupiter.Startup.cs b/Engine/Source/Programs/UnrealCloudDDC/Jupiter/Jupiter.Startup.cs
--- a/Engine/Source/Programs/UnrealCloudDDC/Jupiter/Jupiter.Startup.cs
+++ b/Engine/Source/Programs/UnrealCloudDDC/Jupiter/Jupiter.Startup.cs
@@ -459,11 +459,20 @@ namespace Jupiter
 			{
 				return ActivatorUtilities.CreateInstance<KubernetesLeaderElection>(provider);
 			}
-			else
+			else if (settings.LeaderElectionImplementation == UnrealCloudDDCSettings.LeaderElectionImplementations.Static)
 			{
 				// hard coded leader election that assumes it is always the leader
 				return ActivatorUtilities.CreateInstance<LeaderElectionStub>(provider, true);
 			}
+			else if (settings.LeaderElectionImplementation == UnrealCloudDDCSettings.LeaderElectionImplementations.Disabled)
+			{
+				// disabled leader election means we are never the leader
+				return ActivatorUtilities.CreateInstance<LeaderElectionStub>(provider, false);
+			}
+			else
+			{
+				throw new NotImplementedException($"Unknown leader election set {settings.LeaderElectionImplementation}");
+			}
 		}
 
 		private IAmazonS3 CreateS3(IServiceProvider provider)
diff --git a/Engine/Source/Programs/UnrealCloudDDC/Jupiter/Settings.cs b/Engine/Source/Programs/UnrealCloudDDC/Jupiter/Settings.cs
--- a/Engine/Source/Programs/UnrealCloudDDC/Jupiter/Settings.cs
+++ b/Engine/Source/Programs/UnrealCloudDDC/Jupiter/Settings.cs
@@ -56,7 +56,8 @@ namespace Jupiter
 		public enum LeaderElectionImplementations
 		{
 			Static, 
-			Kubernetes
+			Kubernetes,
+			Disabled
 		}
 
 		public enum ServiceDiscoveryImplementations
