Unreal Cloud DDC - Increased send timeout in nginx to allow for large file uploads (proxy send timeout was not enough)

diff --git a/Engine/Source/Programs/UnrealCloudDDC/Helm/UnrealCloudDDC/templates/configmap-nginx.yaml b/Engine/Source/Programs/UnrealCloudDDC/Helm/UnrealCloudDDC/templates/configmap-nginx.yaml
--- a/Engine/Source/Programs/UnrealCloudDDC/Helm/UnrealCloudDDC/templates/configmap-nginx.yaml
+++ b/Engine/Source/Programs/UnrealCloudDDC/Helm/UnrealCloudDDC/templates/configmap-nginx.yaml
@@ -29,6 +29,7 @@ data:
       # set communication timeout to 5 minutes, transfering really large files takes a while
       proxy_read_timeout 300s;
       proxy_send_timeout 300s;
+      send_timeout 300s;
 
       upstream jupiter_backend {
         zone upstreams 64K;
