Unreal Cloud DDC - Fixed behavior when a replication hasn't run for a long time and thus the offset to resume from is to old (we do not have it anymore). Resetting to the oldest timestamp we have (same behavior as when you enabling a replication for the first time).

diff --git a/Engine/Source/Programs/UnrealCloudDDC/Jupiter/Implementation/TransactionLog/ScyllaReplicationLog.cs b/Engine/Source/Programs/UnrealCloudDDC/Jupiter/Implementation/TransactionLog/ScyllaReplicationLog.cs
--- a/Engine/Source/Programs/UnrealCloudDDC/Jupiter/Implementation/TransactionLog/ScyllaReplicationLog.cs
+++ b/Engine/Source/Programs/UnrealCloudDDC/Jupiter/Implementation/TransactionLog/ScyllaReplicationLog.cs
@@ -195,11 +195,14 @@ namespace Jupiter.Implementation
 				startBucketTime = DateTime.FromFileTimeUtc(bucket);
 				if (startBucketTime < oldCutoff)
 				{
-					// attempting to use a old bucket, this will not exist anymore so will break here
-					yield break;
+					// attempting to use a old bucket, this will not exist anymore so we reset back to the oldest timestamp we have
+					DateTime oldestTimestamp = DateTime.UtcNow.AddSeconds(-1 * _settings.CurrentValue.ReplicationLogTimeToLive.TotalSeconds);
+					startBucketTime = oldestTimestamp;
+				}
+				else
+				{
+					yield return bucket;
 				}
-
-				yield return bucket;
 			}
 			else
 			{
