Unreal Cloud DDC - Simplified query to determine if replication bucket exists, should reduce pressure on ScyllaDB

diff --git a/Engine/Source/Programs/UnrealCloudDDC/Jupiter/Implementation/TransactionLog/ScyllaReplicationLog.cs b/Engine/Source/Programs/UnrealCloudDDC/Jupiter/Implementation/TransactionLog/ScyllaReplicationLog.cs
--- a/Engine/Source/Programs/UnrealCloudDDC/Jupiter/Implementation/TransactionLog/ScyllaReplicationLog.cs
+++ b/Engine/Source/Programs/UnrealCloudDDC/Jupiter/Implementation/TransactionLog/ScyllaReplicationLog.cs
@@ -218,8 +218,7 @@ namespace Jupiter.Implementation
 			{
 				using TelemetrySpan determineBucketExistsScope = _tracer.BuildScyllaSpan("scylla.determine_replication_bucket_exists");
 				// fetch all the buckets that exists and sort them based on time
-				IEnumerable<ScyllaReplicationLogEvent> logEvent = await _mapper.FetchAsync<ScyllaReplicationLogEvent>("WHERE namespace = ? AND replication_bucket = ? LIMIT 1", ns.ToString(), bucketTime.ToFileTimeUtc());
-				ScyllaReplicationLogEvent? e = logEvent.FirstOrDefault();
+				ScyllaReplicationLogEvent? e = await _mapper.FirstOrDefaultAsync<ScyllaReplicationLogEvent>("WHERE namespace = ? AND replication_bucket = ?", ns.ToString(), bucketTime.ToFileTimeUtc());
 				if (e != null)
 				{
 					yield return e.ReplicationBucket;
