Unreal Cloud DDC - Allow for timeouts while determining if a replication log bucket exists, considering any failure as a non-existant bucket. If it fails due to timeout its usually because its scanning all records and not finding any, so its very likely it does not exist.

diff --git a/Engine/Source/Programs/UnrealCloudDDC/Jupiter/Implementation/TransactionLog/ScyllaReplicationLog.cs b/Engine/Source/Programs/UnrealCloudDDC/Jupiter/Implementation/TransactionLog/ScyllaReplicationLog.cs
--- a/Engine/Source/Programs/UnrealCloudDDC/Jupiter/Implementation/TransactionLog/ScyllaReplicationLog.cs
+++ b/Engine/Source/Programs/UnrealCloudDDC/Jupiter/Implementation/TransactionLog/ScyllaReplicationLog.cs
@@ -218,12 +218,20 @@ namespace Jupiter.Implementation
 			{
 				using TelemetrySpan determineBucketExistsScope = _tracer.BuildScyllaSpan("scylla.determine_replication_bucket_exists");
 				// fetch all the buckets that exists and sort them based on time
-				ScyllaReplicationLogEvent? e = await _mapper.FirstOrDefaultAsync<ScyllaReplicationLogEvent>("WHERE namespace = ? AND replication_bucket = ?", ns.ToString(), bucketTime.ToFileTimeUtc());
+				ScyllaReplicationLogEvent? e = null;
+				try
+				{
+					e = await _mapper.FirstOrDefaultAsync<ScyllaReplicationLogEvent>("WHERE namespace = ? AND replication_bucket = ?", ns.ToString(), bucketTime.ToFileTimeUtc());
+				}
+				catch (ReadTimeoutException)
+				{
+					// if the request times out we scanned the db and didn't find anything, we can just move on to the next bucket
+				}
+
 				if (e != null)
 				{
 					yield return e.ReplicationBucket;
 				}
-
 				bucketTime = bucketTime.AddHours(1.0);
 			}
 		}
