Unreal Cloud DDC - Attempts at improving behavior when bad request responses happen during replicator, likely caused by .NET 8 upgrade causing casing differences in the problem details response.

diff --git a/Engine/Source/Programs/UnrealCloudDDC/Jupiter/Implementation/Replication/RefsReplicator.cs b/Engine/Source/Programs/UnrealCloudDDC/Jupiter/Implementation/Replication/RefsReplicator.cs
--- a/Engine/Source/Programs/UnrealCloudDDC/Jupiter/Implementation/Replication/RefsReplicator.cs
+++ b/Engine/Source/Programs/UnrealCloudDDC/Jupiter/Implementation/Replication/RefsReplicator.cs
@@ -675,7 +675,7 @@ namespace Jupiter.Implementation
 				string body = await response.Content.ReadAsStringAsync(cancellationToken);
 				if (response.StatusCode == HttpStatusCode.BadRequest)
 				{
-					ProblemDetails? problemDetails = JsonSerializer.Deserialize<ProblemDetails>(body);
+					ProblemDetails? problemDetails = JsonSerializer.Deserialize<ProblemDetails>(body, DefaultSerializerSettings);
 					if (problemDetails == null)
 					{
 						throw new Exception($"Unknown bad request body when reading incremental replication log. Body: {body}");
@@ -683,7 +683,7 @@ namespace Jupiter.Implementation
 
 					if (problemDetails.Type == ProblemTypes.UseSnapshot)
 					{
-						ProblemDetailsWithSnapshots? problemDetailsWithSnapshots = JsonSerializer.Deserialize<ProblemDetailsWithSnapshots>(body);
+						ProblemDetailsWithSnapshots? problemDetailsWithSnapshots = JsonSerializer.Deserialize<ProblemDetailsWithSnapshots>(body, DefaultSerializerSettings);
 
 						if (problemDetailsWithSnapshots == null)
 						{
@@ -694,6 +694,8 @@ namespace Jupiter.Implementation
 						NamespaceId? blobNamespace = problemDetailsWithSnapshots.BlobNamespace;
 						throw new UseSnapshotException(snapshotBlob, blobNamespace!.Value);
 					}
+				
+					throw new Exception($"Unknown bad request response. Body: {body}");
 				}
 
 				response.EnsureSuccessStatusCode();
