Unreal Cloud DDC - Increased nginx read timeouts to be consistent with send timeouts and increased them all to 10 minutes (as we have seen on demand replication take up to 7 minutes for large shader libraries). This can hopefully be reduced again when we have fixed the shader library to be multiple smaller files.

diff --git a/Engine/Source/Programs/UnrealCloudDDC/Helm/UnrealCloudDDC/templates/configmap-nginx.yaml b/Engine/Source/Programs/UnrealCloudDDC/Helm/UnrealCloudDDC/templates/configmap-nginx.yaml
--- a/Engine/Source/Programs/UnrealCloudDDC/Helm/UnrealCloudDDC/templates/configmap-nginx.yaml
+++ b/Engine/Source/Programs/UnrealCloudDDC/Helm/UnrealCloudDDC/templates/configmap-nginx.yaml
@@ -26,10 +26,11 @@ data:
       sendfile on;
       tcp_nopush on;
 
-      # set communication timeout to 5 minutes, transfering really large files takes a while
-      proxy_read_timeout 300s;
-      proxy_send_timeout 300s;
-      send_timeout 300s;
+      # set communication timeout to 10 minutes, transfering really large files takes a while
+      proxy_read_timeout 600s;
+      proxy_send_timeout 600s;
+      read_timeout 600s;
+      send_timeout 600s;
 
       upstream jupiter_backend {
         zone upstreams 64K;
