Unreal Cloud DDC - Nginx seperates read timeouts into two seperate timeouts, bumping both to be consistent with our other timeouts.
Hopefully resolves timeouts during oplog replication

diff --git a/Engine/Source/Programs/UnrealCloudDDC/Helm/UnrealCloudDDC/templates/configmap-nginx.yaml b/Engine/Source/Programs/UnrealCloudDDC/Helm/UnrealCloudDDC/templates/configmap-nginx.yaml
--- a/Engine/Source/Programs/UnrealCloudDDC/Helm/UnrealCloudDDC/templates/configmap-nginx.yaml
+++ b/Engine/Source/Programs/UnrealCloudDDC/Helm/UnrealCloudDDC/templates/configmap-nginx.yaml
@@ -29,7 +29,8 @@ data:
       # set communication timeout to 10 minutes, transfering really large files takes a while
       proxy_read_timeout 600s;
       proxy_send_timeout 600s;
-      read_timeout 600s;
+      client_header_timeout 600s;
+      client_body_timeout 600s;
       send_timeout 600s;
 
       upstream jupiter_backend {
