Unreal Cloud DDC - Removed comment about documentation, while we do still (under certain conditions) provide swagger documentation we do not host that under /docs as this was done via our frontend (that no longer exists). Generally we provide documentation in the form of readmes

diff --git a/Engine/Source/Programs/UnrealCloudDDC/README.md b/Engine/Source/Programs/UnrealCloudDDC/README.md
--- a/Engine/Source/Programs/UnrealCloudDDC/README.md
+++ b/Engine/Source/Programs/UnrealCloudDDC/README.md
@@ -73,9 +73,6 @@ The AWS compose file can be replaced with `docker-compose-azure.yml` if you want
 
 Docker compose setups disable authentication for to make it quick to get started, generally we recommend that you hook UnrealCloudDDC up to a OIDC provider before deploying this.
 
-UnrealCloudDDC hosts Swagger documentation at `/docs` (so `http://localhost/docs` when running locally). This page lets you pick the service you want to call and lists the API for it.
-Note that some endpoints using binary protocol for effieceny and the format of those is not documented in swagger, we always provide a REST api for every use case.
-
 # Deployment
 UnrealCloudDDC is currently only run in production on AWS, but the requirements on storage and db are very generic and also abstracted. 
 We have basic (untested) support for Azure services.
