Unreal Cloud DDC - Use kubernetes default periodic frequency for liveness and readiness probes and increased the timeout a bit to reduce noise

diff --git a/Engine/Source/Programs/UnrealCloudDDC/Helm/UnrealCloudDDC/values.yaml b/Engine/Source/Programs/UnrealCloudDDC/Helm/UnrealCloudDDC/values.yaml
--- a/Engine/Source/Programs/UnrealCloudDDC/Helm/UnrealCloudDDC/values.yaml
+++ b/Engine/Source/Programs/UnrealCloudDDC/Helm/UnrealCloudDDC/values.yaml
@@ -120,15 +120,15 @@ resources:
 livenessProbe:
   enabled: true
   initialDelaySeconds: 5
-  periodSeconds: 5
+  periodSeconds: 10
   timeoutSeconds: 5
   successThreshold: 1
   failureThreshold: 5
 readinessProbe:
   enabled: true
   initialDelaySeconds: 5
-  periodSeconds: 5
-  timeoutSeconds: 1
+  periodSeconds: 10
+  timeoutSeconds: 3
   successThreshold: 1
   failureThreshold: 5
 
@@ -336,15 +336,15 @@ nginx:
   livenessProbe:
     enabled: true
     initialDelaySeconds: 5
-    periodSeconds: 5
+    periodSeconds: 10
     timeoutSeconds: 5
     successThreshold: 1
     failureThreshold: 5
   readinessProbe:
     enabled: true
     initialDelaySeconds: 5
-    periodSeconds: 5
-    timeoutSeconds: 1
+    periodSeconds: 10
+    timeoutSeconds: 3
     successThreshold: 1
     failureThreshold: 5
 
@@ -371,15 +371,15 @@ worker:
   livenessProbe:
     enabled: true
     initialDelaySeconds: 5
-    periodSeconds: 5
+    periodSeconds: 10
     timeoutSeconds: 5
     successThreshold: 1
     failureThreshold: 5
   readinessProbe:
     enabled: true
     initialDelaySeconds: 5
-    periodSeconds: 5
-    timeoutSeconds: 1
+    periodSeconds: 10
+    timeoutSeconds: 3
     successThreshold: 1
     failureThreshold: 5
     
