Unreal Cloud DDC - Fixed incorrect syntax in compose for azure specific resources

diff --git a/Engine/Source/Programs/UnrealCloudDDC/Composes/docker-compose-azure.yml b/Engine/Source/Programs/UnrealCloudDDC/Composes/docker-compose-azure.yml
--- a/Engine/Source/Programs/UnrealCloudDDC/Composes/docker-compose-azure.yml
+++ b/Engine/Source/Programs/UnrealCloudDDC/Composes/docker-compose-azure.yml
@@ -5,7 +5,7 @@ services:
     - UnrealCloudDDC__storageImplementations__1=azure
     - Azure__ConnectionString=DefaultEndpointsProtocol=https;BlobEndpoint=http://azure_storage:11002/blobaccount;AccountName=blobaccount;AccountKey=iU6uTvlF1ysppmft+NO5lAD0E3hwrAORr5Rb5xcBWUgEz/OicrSkFxwZYMNK5XL29/wXZKGOoOVSW040nAOfPg==
     
-    network:
+    networks:
     - network_azure_storage
   azure_storage:
     image: "mcr.microsoft.com/azure-blob-storage"
