Unreal Cloud DDC - Added runtime instrumentation to collect the misisng .net metrics we used to get from datadog.
Also set the sampler for tracing to always on - we will filter this in later stages if we want to limit it - we generally rely heavily on tracing when following up performance issues so its important we can rely on all requests being in the trace.

diff --git a/Engine/Source/Programs/UnrealCloudDDC/Jupiter.Common/BaseStartup.cs b/Engine/Source/Programs/UnrealCloudDDC/Jupiter.Common/BaseStartup.cs
--- a/Engine/Source/Programs/UnrealCloudDDC/Jupiter.Common/BaseStartup.cs
+++ b/Engine/Source/Programs/UnrealCloudDDC/Jupiter.Common/BaseStartup.cs
@@ -298,7 +298,7 @@ namespace Jupiter
 						}
 					};
 				});
-
+				builder.SetSampler(new AlwaysOnSampler());
 				builder.AddOtlpExporter();
 				if (useConsoleExporter)
 				{
@@ -311,7 +311,8 @@ namespace Jupiter
 					.AddMeter("UnrealCloudDDC", "ScyllaDB")
 					.AddOtlpExporter()
 					.AddAspNetCoreInstrumentation()
-					.AddHttpClientInstrumentation();
+					.AddHttpClientInstrumentation()
+					.AddRuntimeInstrumentation();
 
 				if (useConsoleExporter)
 				{
diff --git a/Engine/Source/Programs/UnrealCloudDDC/Jupiter.Common/Jupiter.Common.csproj b/Engine/Source/Programs/UnrealCloudDDC/Jupiter.Common/Jupiter.Common.csproj
--- a/Engine/Source/Programs/UnrealCloudDDC/Jupiter.Common/Jupiter.Common.csproj
+++ b/Engine/Source/Programs/UnrealCloudDDC/Jupiter.Common/Jupiter.Common.csproj
@@ -37,6 +37,7 @@
     <PackageReference Include="OpenTelemetry.Exporter.OpenTelemetryProtocol" Version="1.7.0" />
     <PackageReference Include="OpenTelemetry.Instrumentation.AspNetCore" Version="1.7.0" />
     <PackageReference Include="OpenTelemetry.Instrumentation.Http" Version="1.7.0" />
+    <PackageReference Include="OpenTelemetry.Instrumentation.Runtime" Version="1.7.0" />
     <PackageReference Include="OpenTelemetry.Extensions.Hosting" Version="1.7.0" />
     <PackageReference Include="IdentityServer4.AccessTokenValidation" Version="3.0.1" />
 
