Unreal Cloud DDC - Removing increased nginx timeouts as this causes very long stalls in DDC requests in us-east-1, still determining why these requests are so slow.

diff --git a/Engine/Source/Programs/UnrealCloudDDC/Helm/UnrealCloudDDC/templates/configmap-nginx.yaml b/Engine/Source/Programs/UnrealCloudDDC/Helm/UnrealCloudDDC/templates/configmap-nginx.yaml
--- a/Engine/Source/Programs/UnrealCloudDDC/Helm/UnrealCloudDDC/templates/configmap-nginx.yaml
+++ b/Engine/Source/Programs/UnrealCloudDDC/Helm/UnrealCloudDDC/templates/configmap-nginx.yaml
@@ -26,13 +26,6 @@ data:
       sendfile on;
       tcp_nopush on;
 
-      # set communication timeout to 10 minutes, transfering really large files takes a while
-      proxy_read_timeout 600s;
-      proxy_send_timeout 600s;
-      client_header_timeout 600s;
-      client_body_timeout 600s;
-      send_timeout 600s;
-
       upstream jupiter_backend {
         zone upstreams 64K;
         server {{ if .Values.nginx.useDomainSockets -}} unix:/nginx/jupiter-http.sock {{- else -}} localhost:80 {{- end }} max_fails=1 fail_timeout=2s;
