Unreal Cloud DDC - Rent buffer used to hash a stream, prevents churn in the GC as this buffer is a fixed size

diff --git a/Engine/Source/Programs/UnrealCloudDDC/Jupiter.Common/Implementation/BlobId.cs b/Engine/Source/Programs/UnrealCloudDDC/Jupiter.Common/Implementation/BlobId.cs
--- a/Engine/Source/Programs/UnrealCloudDDC/Jupiter.Common/Implementation/BlobId.cs
+++ b/Engine/Source/Programs/UnrealCloudDDC/Jupiter.Common/Implementation/BlobId.cs
@@ -1,6 +1,7 @@
 ﻿// Copyright Epic Games, Inc. All Rights Reserved.
 
 using System;
+using System.Buffers;
 using System.ComponentModel;
 using System.Globalization;
 using System.IO;
@@ -127,19 +128,28 @@ namespace Jupiter.Implementation
 		public static async Task<BlobId> FromStreamAsync(Stream stream)
 		{
 			using Hasher hasher = Hasher.New();
-			const int bufferSize = 1024 * 1024 * 5;
-			byte[] buffer = new byte[bufferSize];
-			int read = await stream.ReadAsync(buffer, 0, buffer.Length);
-			while (read > 0)
+
+			const int BufferSize = 1024 * 1024 * 5;
+			byte[] buffer = ArrayPool<byte>.Shared.Rent(BufferSize);
+
+			try
 			{
-				hasher.UpdateWithJoin(new ReadOnlySpan<byte>(buffer, 0, read));
-				read = await stream.ReadAsync(buffer, 0, buffer.Length);
+				int read = await stream.ReadAsync(buffer, 0, buffer.Length);
+				while (read > 0)
+				{
+					hasher.UpdateWithJoin(new ReadOnlySpan<byte>(buffer, 0, read));
+					read = await stream.ReadAsync(buffer, 0, buffer.Length);
+				}
+				Hash blake3Hash = hasher.Finalize();
+
+				// we only keep the first 20 bytes of the Blake3 hash
+				byte[] hash = blake3Hash.AsSpan().Slice(0, 20).ToArray();
+				return new BlobId(hash);
+			}
+			finally
+			{
+				ArrayPool<byte>.Shared.Return(buffer);
 			}
-			Hash blake3Hash = hasher.Finalize();
-
-			// we only keep the first 20 bytes of the Blake3 hash
-			byte[] hash = blake3Hash.AsSpan().Slice(0, 20).ToArray();
-			return new BlobId(hash);
 		}
 
 		public static BlobId FromIoHash(IoHash blobIdentifier)
