Unreal Cloud DDC - Added option to not throw when blobs are enumerated and instead just returns the ones we found - and enabled that when GC is updating stats tracking.
Resolves occasional exceptions thrown during GC resulting on a ref not getting deleted.

diff --git a/Engine/Source/Programs/UnrealCloudDDC/Jupiter/Implementation/GC/CacheCleanup.cs b/Engine/Source/Programs/UnrealCloudDDC/Jupiter/Implementation/GC/CacheCleanup.cs
--- a/Engine/Source/Programs/UnrealCloudDDC/Jupiter/Implementation/GC/CacheCleanup.cs
+++ b/Engine/Source/Programs/UnrealCloudDDC/Jupiter/Implementation/GC/CacheCleanup.cs
@@ -170,7 +170,7 @@ namespace Jupiter.Implementation
 			{
 				if (_cloudDDCSettings.CurrentValue.EnableBucketStatsTracking)
 				{
-					List<BlobId> blobs = await _objectService.GetReferencedBlobsAsync(ns, bucket, name);
+					List<BlobId> blobs = await _objectService.GetReferencedBlobsAsync(ns, bucket, name, ignoreMissingBlobs: true);
 					await _blobIndex.RemoveBlobFromBucketListAsync(ns, bucket, name, blobs);
 				}
 
diff --git a/Engine/Source/Programs/UnrealCloudDDC/Jupiter/Implementation/IRefService.cs b/Engine/Source/Programs/UnrealCloudDDC/Jupiter/Implementation/IRefService.cs
--- a/Engine/Source/Programs/UnrealCloudDDC/Jupiter/Implementation/IRefService.cs
+++ b/Engine/Source/Programs/UnrealCloudDDC/Jupiter/Implementation/IRefService.cs
@@ -21,7 +21,7 @@ namespace Jupiter.Implementation
 		Task<long> DeleteBucketAsync(NamespaceId ns, BucketId bucket);
 
 		Task<bool> ExistsAsync(NamespaceId ns, BucketId bucket, RefId key);
-		Task<List<BlobId>> GetReferencedBlobsAsync(NamespaceId ns, BucketId bucket, RefId key);
+		Task<List<BlobId>> GetReferencedBlobsAsync(NamespaceId ns, BucketId bucket, RefId key, bool ignoreMissingBlobs);
 	}
 
 	public class ObjectHashMismatchException : Exception
diff --git a/Engine/Source/Programs/UnrealCloudDDC/Jupiter/Implementation/IReferenceResolver.cs b/Engine/Source/Programs/UnrealCloudDDC/Jupiter/Implementation/IReferenceResolver.cs
--- a/Engine/Source/Programs/UnrealCloudDDC/Jupiter/Implementation/IReferenceResolver.cs
+++ b/Engine/Source/Programs/UnrealCloudDDC/Jupiter/Implementation/IReferenceResolver.cs
@@ -73,8 +73,9 @@ namespace Jupiter.Implementation
 		/// </summary>
 		/// <param name="ns">The namespace to check</param>
 		/// <param name="cb">The compact binary object to resolve references for</param>
+		/// <param name="ignoreMissingBlobs">Set to true to always returned the blobs found, ignoring anything that is missing rather then throwing</param>
 		/// <returns></returns>
-		IAsyncEnumerable<BlobId> GetReferencedBlobs(NamespaceId ns, CbObject cb);
+		IAsyncEnumerable<BlobId> GetReferencedBlobs(NamespaceId ns, CbObject cb, bool ignoreMissingBlobs = false);
 
 		/// <summary>
 		/// Returns which attachments exist in the cb object or any children
@@ -233,7 +234,7 @@ namespace Jupiter.Implementation
 			}
 		}
 
-		public async IAsyncEnumerable<BlobId> GetReferencedBlobs(NamespaceId ns, CbObject cb)
+		public async IAsyncEnumerable<BlobId> GetReferencedBlobs(NamespaceId ns, CbObject cb, bool ignoreMissingBlobs = false)
 		{
 			List<Task<(BlobId, bool)>> pendingBlobExistsChecks = new();
 			List<Task<(ContentIdAttachment, bool)>> pendingContentIdChecks = new();
@@ -294,12 +295,12 @@ namespace Jupiter.Implementation
 			}
 
 			// if there were any content ids we did not recognize we throw a partial reference exception
-			if (unresolvedContentIdReferences.Count != 0)
+			if (!ignoreMissingBlobs && unresolvedContentIdReferences.Count != 0)
 			{
 				throw new PartialReferenceResolveException(unresolvedContentIdReferences);
 			}
 			// if there were any blobs missing we throw a partial reference exception
-			if (unresolvedBlobReferences.Count != 0)
+			if (!ignoreMissingBlobs && unresolvedBlobReferences.Count != 0)
 			{
 				throw new ReferenceIsMissingBlobsException(unresolvedBlobReferences);
 			}
diff --git a/Engine/Source/Programs/UnrealCloudDDC/Jupiter/Implementation/RefService.cs b/Engine/Source/Programs/UnrealCloudDDC/Jupiter/Implementation/RefService.cs
--- a/Engine/Source/Programs/UnrealCloudDDC/Jupiter/Implementation/RefService.cs
+++ b/Engine/Source/Programs/UnrealCloudDDC/Jupiter/Implementation/RefService.cs
@@ -312,7 +312,7 @@ namespace Jupiter.Implementation
 			return true;
 		}
 
-		public async Task<List<BlobId>> GetReferencedBlobsAsync(NamespaceId ns, BucketId bucket, RefId name)
+		public async Task<List<BlobId>> GetReferencedBlobsAsync(NamespaceId ns, BucketId bucket, RefId name, bool ignoreMissingBlobs = false)
 		{
 			byte[] blob;
 			RefRecord o = await _referencesStore.GetAsync(ns, bucket, name, IReferencesStore.FieldFlags.IncludePayload, IReferencesStore.OperationFlags.None);
@@ -328,7 +328,7 @@ namespace Jupiter.Implementation
 
 			CbObject cbObject = new CbObject(blob);
 
-			List<BlobId> referencedBlobs = await _referenceResolver.GetReferencedBlobs(ns, cbObject).ToListAsync();
+			List<BlobId> referencedBlobs = await _referenceResolver.GetReferencedBlobs(ns, cbObject, ignoreMissingBlobs).ToListAsync();
 			return referencedBlobs;
 		}
 	}
