Unreal Cloud DDC - Added logging when filesystem cleanup runs to try and determine why the cleanup in EUC is deciding not to run resulting in out of diskspace issues as well as some more tracing info to better understand what we spend time on.

diff --git a/Engine/Source/Programs/UnrealCloudDDC/Jupiter/Implementation/Blob/FileSystemStore.cs b/Engine/Source/Programs/UnrealCloudDDC/Jupiter/Implementation/Blob/FileSystemStore.cs
--- a/Engine/Source/Programs/UnrealCloudDDC/Jupiter/Implementation/Blob/FileSystemStore.cs
+++ b/Engine/Source/Programs/UnrealCloudDDC/Jupiter/Implementation/Blob/FileSystemStore.cs
@@ -21,12 +21,14 @@ namespace Jupiter.Implementation
 		private readonly IServiceProvider _provider;
 		private readonly IOptionsMonitor<FilesystemSettings> _settings;
 		private readonly Tracer _tracer;
+		private readonly ILogger<FileSystemStore> _logger;
 		private readonly ConcurrentDictionary<NamespaceId, FileStorageBackend> _backends = new ConcurrentDictionary<NamespaceId, FileStorageBackend>();
 
-		public FileSystemStore(IOptionsMonitor<FilesystemSettings> settings, Tracer tracer, IServiceProvider provider)
+		public FileSystemStore(IOptionsMonitor<FilesystemSettings> settings, Tracer tracer, ILogger<FileSystemStore> logger, IServiceProvider provider)
 		{
 			_settings = settings;
 			_tracer = tracer;
+			_logger = logger;
 			_provider = provider;
 		}
 
@@ -173,15 +175,19 @@ namespace Jupiter.Implementation
 				// first check to see if we should trigger at all, this happens for each run but only really matters for the first attempt
 				if (size < triggerSize)
 				{
+					_logger.LogInformation("Filesystem cleanup not running. Disksize used: {UsedDiskSize} . Trigger size was {TriggerSize}", size, triggerSize);
 					return countOfBlobsRemoved;
 				}
 
 				// then check if we have reached the target size, if not we should continue running
 				if (size <= targetSize)
 				{
+					_logger.LogInformation("Filesystem cleanup reached target size. Disksize used: {UsedDiskSize} . Target size was {TargetSize}", size, targetSize);
+
 					return countOfBlobsRemoved;
 				}
 				
+				_logger.LogInformation("Filesystem cleanup running. Disksize used: {UsedDiskSize} . Trigger size was {TriggerSize}", size, triggerSize);
 				IEnumerable<FileInfo> fileInfos = GetLeastRecentlyAccessedObjects(maxResults: batchSize);
 
 				bool hadFiles = false;
@@ -244,6 +250,9 @@ namespace Jupiter.Implementation
 		/// <returns>Total size of blobs in bytes</returns>
 		public async Task<long> CalculateDiskSpaceUsedAsync(NamespaceId? ns = null)
 		{
+			using TelemetrySpan scope = _tracer.StartActiveSpan("gc.filesystem.calc_disc_usage")
+				.SetAttribute("operation.name", "gc.filesystem.calc_disc_usage");
+
 			string path = ns != null ? Path.Combine(GetRootDir(), ns.ToString()!) : GetRootDir();
 			DirectoryInfo di = new DirectoryInfo(path);
 			if (!di.Exists)
