Unreal Cloud DDC - Fixed issue with the consistency check for the ref store, was enumerating data from the wrong table resulting in always seeing consistency (as it was checking consistency with the same table instead of comparing tables).

diff --git a/Engine/Source/Programs/UnrealCloudDDC/Jupiter/Implementation/Blob/RefStoreConsistencyCheckService.cs b/Engine/Source/Programs/UnrealCloudDDC/Jupiter/Implementation/Blob/RefStoreConsistencyCheckService.cs
--- a/Engine/Source/Programs/UnrealCloudDDC/Jupiter/Implementation/Blob/RefStoreConsistencyCheckService.cs
+++ b/Engine/Source/Programs/UnrealCloudDDC/Jupiter/Implementation/Blob/RefStoreConsistencyCheckService.cs
@@ -70,7 +70,7 @@ namespace Jupiter.Implementation
 
 			ulong countOfRefsChecked = 0;
 			ulong countOfMissingLastAccessTime = 0;
-			await foreach ((NamespaceId ns, BucketId bucket, RefId refId, DateTime lastAccess) in _referencesStore.GetRecordsAsync())
+			await foreach ((NamespaceId ns, BucketId bucket, RefId refId) in _referencesStore.GetRecordsWithoutAccessTimeAsync())
 			{
 				using TelemetrySpan scope = _tracer.StartActiveSpan("consistency_check.ref_store")
 					.SetAttribute("operation.name", "consistency_check.ref_store")
diff --git a/Engine/Source/Programs/UnrealCloudDDC/Jupiter/Implementation/Objects/CacheReferencesStore.cs b/Engine/Source/Programs/UnrealCloudDDC/Jupiter/Implementation/Objects/CacheReferencesStore.cs
--- a/Engine/Source/Programs/UnrealCloudDDC/Jupiter/Implementation/Objects/CacheReferencesStore.cs
+++ b/Engine/Source/Programs/UnrealCloudDDC/Jupiter/Implementation/Objects/CacheReferencesStore.cs
@@ -97,6 +97,11 @@ namespace Jupiter.Implementation
 			throw new NotImplementedException("GetRecords not supported on a cached reference store");
 		}
 
+		public IAsyncEnumerable<(NamespaceId, BucketId, RefId)> GetRecordsWithoutAccessTimeAsync()
+		{
+			throw new NotImplementedException("GetRecordsWithoutAccessTimeAsync not supported on a cached reference store");
+		}
+
 		public IAsyncEnumerable<(RefId, BlobId)> GetRecordsInBucketAsync(NamespaceId ns, BucketId bucket)
 		{
 			throw new NotImplementedException("GetRecordsInBucketAsync not supported on a cached reference store");
@@ -234,6 +239,11 @@ namespace Jupiter.Implementation
 			throw new NotImplementedException("GetRecords is not supported on a upstream reference store");
 		}
 
+		public IAsyncEnumerable<(NamespaceId, BucketId, RefId)> GetRecordsWithoutAccessTimeAsync()
+		{
+			throw new NotImplementedException("GetRecordsWithoutAccessTimeAsync is not supported on a upstream reference store");
+		}
+
 		public IAsyncEnumerable<(RefId, BlobId)> GetRecordsInBucketAsync(NamespaceId ns, BucketId bucket)
 		{
 			throw new NotImplementedException("GetRecordsInBucketAsync is not supported on a upstream reference store");
diff --git a/Engine/Source/Programs/UnrealCloudDDC/Jupiter/Implementation/Objects/IReferencesStore.cs b/Engine/Source/Programs/UnrealCloudDDC/Jupiter/Implementation/Objects/IReferencesStore.cs
--- a/Engine/Source/Programs/UnrealCloudDDC/Jupiter/Implementation/Objects/IReferencesStore.cs
+++ b/Engine/Source/Programs/UnrealCloudDDC/Jupiter/Implementation/Objects/IReferencesStore.cs
@@ -33,6 +33,8 @@ namespace Jupiter.Implementation
 		Task UpdateLastAccessTimeAsync(NamespaceId ns, BucketId bucket, RefId key, DateTime newLastAccessTime);
 		IAsyncEnumerable<(NamespaceId, BucketId, RefId, DateTime)> GetRecordsAsync();
 
+		IAsyncEnumerable<(NamespaceId, BucketId, RefId)> GetRecordsWithoutAccessTimeAsync();
+
 		IAsyncEnumerable<(RefId, BlobId)> GetRecordsInBucketAsync(NamespaceId ns, BucketId bucket);
 
 		IAsyncEnumerable<NamespaceId> GetNamespacesAsync();
diff --git a/Engine/Source/Programs/UnrealCloudDDC/Jupiter/Implementation/Objects/MemoryCachedReferencesStore.cs b/Engine/Source/Programs/UnrealCloudDDC/Jupiter/Implementation/Objects/MemoryCachedReferencesStore.cs
--- a/Engine/Source/Programs/UnrealCloudDDC/Jupiter/Implementation/Objects/MemoryCachedReferencesStore.cs
+++ b/Engine/Source/Programs/UnrealCloudDDC/Jupiter/Implementation/Objects/MemoryCachedReferencesStore.cs
@@ -117,6 +117,11 @@ namespace Jupiter.Implementation.Objects
 			return _actualStore.GetRecordsAsync();
 		}
 
+		public IAsyncEnumerable<(NamespaceId, BucketId, RefId)> GetRecordsWithoutAccessTimeAsync()
+		{
+			return _actualStore.GetRecordsWithoutAccessTimeAsync();
+		}
+
 		public IAsyncEnumerable<(RefId, BlobId)> GetRecordsInBucketAsync(NamespaceId ns, BucketId bucket)
 		{
 			return _actualStore.GetRecordsInBucketAsync(ns, bucket);
diff --git a/Engine/Source/Programs/UnrealCloudDDC/Jupiter/Implementation/Objects/MemoryReferencesStore.cs b/Engine/Source/Programs/UnrealCloudDDC/Jupiter/Implementation/Objects/MemoryReferencesStore.cs
--- a/Engine/Source/Programs/UnrealCloudDDC/Jupiter/Implementation/Objects/MemoryReferencesStore.cs
+++ b/Engine/Source/Programs/UnrealCloudDDC/Jupiter/Implementation/Objects/MemoryReferencesStore.cs
@@ -76,6 +76,15 @@ namespace Jupiter.Implementation
 			return Task.CompletedTask;
 		}
 
+		public async IAsyncEnumerable<(NamespaceId, BucketId, RefId)> GetRecordsWithoutAccessTimeAsync()
+		{
+			await foreach ((NamespaceId namespaceId, BucketId bucketId, RefId refId, DateTime _) in GetRecordsAsync())
+			{
+				await Task.CompletedTask;
+				yield return (namespaceId, bucketId, refId);
+			}
+		}
+
 		public async IAsyncEnumerable<(NamespaceId, BucketId, RefId, DateTime)> GetRecordsAsync()
 		{
 			foreach (MemoryStoreObject o in _objects.Values.OrderBy(o => o.LastAccessTime))
diff --git a/Engine/Source/Programs/UnrealCloudDDC/Jupiter/Implementation/Objects/MongoReferencesStore.cs b/Engine/Source/Programs/UnrealCloudDDC/Jupiter/Implementation/Objects/MongoReferencesStore.cs
--- a/Engine/Source/Programs/UnrealCloudDDC/Jupiter/Implementation/Objects/MongoReferencesStore.cs
+++ b/Engine/Source/Programs/UnrealCloudDDC/Jupiter/Implementation/Objects/MongoReferencesStore.cs
@@ -141,6 +141,14 @@ namespace Jupiter.Implementation
 			}
 		}
 
+		public async IAsyncEnumerable<(NamespaceId, BucketId, RefId)> GetRecordsWithoutAccessTimeAsync()
+		{
+			await foreach ((NamespaceId ns, BucketId bucket, RefId key, DateTime _) in GetRecordsAsync())
+			{
+				yield return (ns, bucket, key);
+			}
+		}
+
 		public async IAsyncEnumerable<(RefId, BlobId)> GetRecordsInBucketAsync(NamespaceId ns, BucketId bucket)
 		{
 			IMongoCollection<MongoReferencesModelV0> collection = GetCollection<MongoReferencesModelV0>();
diff --git a/Engine/Source/Programs/UnrealCloudDDC/Jupiter/Implementation/Objects/ScyllaReferencesStore.cs b/Engine/Source/Programs/UnrealCloudDDC/Jupiter/Implementation/Objects/ScyllaReferencesStore.cs
--- a/Engine/Source/Programs/UnrealCloudDDC/Jupiter/Implementation/Objects/ScyllaReferencesStore.cs
+++ b/Engine/Source/Programs/UnrealCloudDDC/Jupiter/Implementation/Objects/ScyllaReferencesStore.cs
@@ -280,6 +280,81 @@ namespace Jupiter.Implementation
 			}
 		}
 
+		
+		public async IAsyncEnumerable<(NamespaceId, BucketId, RefId)> GetRecordsWithoutAccessTimeAsync()
+		{
+			using TelemetrySpan scope = _tracer.BuildScyllaSpan("scylla.get_records_no_access_time");
+
+			if (_settings.CurrentValue.UsePerShardScanning)
+			{
+				IAsyncEnumerable<(NamespaceId, BucketId, RefId, DateTime)> enumerable = GetRecordsPerShardAsync(false);
+
+				await foreach ((NamespaceId, BucketId, RefId, DateTime) record in enumerable)
+				{
+					yield return (record.Item1, record.Item2, record.Item3);
+				}
+			}
+			else
+			{
+				// use the object table for listing here as its used by the consistency check
+				PreparedStatement getObjectStatement = _getObjectsStatement;
+				const int MaxRetryAttempts = 3;
+				RowSet rowSet = await _session.ExecuteAsync(getObjectStatement.Bind());
+
+				do
+				{
+					int countOfRows = rowSet.GetAvailableWithoutFetching();
+					IEnumerable<Row> localRows = rowSet.Take(countOfRows);
+					Task prefetchTask = rowSet.FetchMoreResultsAsync();
+
+					foreach (Row row in localRows)
+					{
+						string ns = row.GetValue<string>("namespace");
+						string bucket = row.GetValue<string>("bucket");
+						string name = row.GetValue<string>("name");
+
+						// skip any names that are not conformant to io hash
+						if (name.Length != 40)
+						{
+							continue;
+						}
+
+						yield return (new NamespaceId(ns), new BucketId(bucket), new RefId(name));
+					}
+
+					int retryAttempts = 0;
+					Exception? timeoutException = null;
+					while (retryAttempts < MaxRetryAttempts)
+					{
+						try
+						{
+							await prefetchTask;
+							timeoutException = null;
+							break;
+						}
+						catch (ReadTimeoutException e)
+						{
+							retryAttempts += 1;
+							_logger.LogWarning(
+								"Cassandra read timeouts, waiting a while and then retrying. Attempt {Attempts} .",
+								retryAttempts);
+							// wait 10 seconds and try again as the Db is under heavy load right now
+							await Task.Delay(TimeSpan.FromSeconds(10));
+							timeoutException = e;
+						}
+					}
+
+					if (timeoutException != null)
+					{
+						_logger.LogWarning("Cassandra read timeouts, attempted {Attempts} attempts now we give up.",
+							retryAttempts);
+						// we have failed to many times, rethrow the exception and abort to avoid stalling here for ever
+						throw timeoutException;
+					}
+				} while (!rowSet.IsFullyFetched);
+			}
+		}
+
 		public async IAsyncEnumerable<(RefId, BlobId)> GetRecordsInBucketAsync(NamespaceId ns, BucketId bucket)
 		{
 			using TelemetrySpan scope = _tracer.BuildScyllaSpan("scylla.get_records_in_bucket_per_shard");
@@ -311,13 +386,17 @@ namespace Jupiter.Implementation
 		/// See https://www.scylladb.com/2017/03/28/parallel-efficient-full-table-scan-scylla/
 		/// </summary>
 		/// <returns></returns>
-		private async IAsyncEnumerable<(NamespaceId, BucketId, RefId, DateTime)> GetRecordsPerShardAsync()
+		private async IAsyncEnumerable<(NamespaceId, BucketId, RefId, DateTime)> GetRecordsPerShardAsync(bool? forceUseLastAccessTable = null)
 		{
 			using TelemetrySpan scope = _tracer.BuildScyllaSpan("scylla.get_records_per_shard");
-			PreparedStatement getObjectStatement = _settings.CurrentValue.ListObjectsFromLastAccessTable
+			
+			bool useLastAccessTable = forceUseLastAccessTable ?? _settings.CurrentValue.ListObjectsFromLastAccessTable;
+			PreparedStatement getObjectStatement = useLastAccessTable
 				? _getObjectsLastAccessForPartitionRangeStatement
 				: _getObjectsForPartitionRangeStatement;
 
+			scope.SetAttribute("TableUsed", useLastAccessTable ? "LastAccess" : "Objects");
+
 			// generate a list of all the primary key ranges that exist on the cluster
 			List<(long, long)> tableRanges = ScyllaUtils.GetTableRanges(_settings.CurrentValue.CountOfNodes, _settings.CurrentValue.CountOfCoresPerNode, 3).ToList();
 
