Unreal Cloud DDC - Improvements to speed when GC refs. First of removed the IF EXISTS check when deleting a ref as it isn't super important to have this information and that triggers much more work on the db side. Parralleized the check with the stats tracking and the actual deletes.

diff --git a/Engine/Source/Programs/UnrealCloudDDC/Jupiter/Implementation/GC/CacheCleanup.cs b/Engine/Source/Programs/UnrealCloudDDC/Jupiter/Implementation/GC/CacheCleanup.cs
--- a/Engine/Source/Programs/UnrealCloudDDC/Jupiter/Implementation/GC/CacheCleanup.cs
+++ b/Engine/Source/Programs/UnrealCloudDDC/Jupiter/Implementation/GC/CacheCleanup.cs
@@ -124,7 +124,7 @@ namespace Jupiter.Implementation
 
 						await DeleteRefAsync(ns, bucket, name);
 					
-							Interlocked.Increment(ref countOfDeletedRecords);
+						Interlocked.Increment(ref countOfDeletedRecords);
 
 						return;
 					}
@@ -168,10 +168,14 @@ namespace Jupiter.Implementation
 			bool storeDelete = false;
 			try
 			{
+				Task? bucketStatsCleanupTask = null;
 				if (_cloudDDCSettings.CurrentValue.EnableBucketStatsTracking)
 				{
-					List<BlobId> blobs = await _objectService.GetReferencedBlobsAsync(ns, bucket, name, ignoreMissingBlobs: true);
-					await _blobIndex.RemoveBlobFromBucketListAsync(ns, bucket, name, blobs);
+					bucketStatsCleanupTask = Task.Run(async () =>
+					{
+						List<BlobId> blobs = await _objectService.GetReferencedBlobsAsync(ns, bucket, name, ignoreMissingBlobs: true);
+						await _blobIndex.RemoveBlobFromBucketListAsync(ns, bucket, name, blobs);
+					});
 				}
 
 				storeDelete = await _referencesStore.DeleteAsync(ns, bucket, name);
@@ -180,6 +184,11 @@ namespace Jupiter.Implementation
 					// insert a delete event into the transaction log
 					await _replicationLog.InsertDeleteEventAsync(ns, bucket, name, null);
 				}
+
+				if (bucketStatsCleanupTask != null)
+				{
+					await bucketStatsCleanupTask;
+				}
 			}
 			catch (Exception e)
 			{
diff --git a/Engine/Source/Programs/UnrealCloudDDC/Jupiter/Implementation/Objects/ScyllaReferencesStore.cs b/Engine/Source/Programs/UnrealCloudDDC/Jupiter/Implementation/Objects/ScyllaReferencesStore.cs
--- a/Engine/Source/Programs/UnrealCloudDDC/Jupiter/Implementation/Objects/ScyllaReferencesStore.cs
+++ b/Engine/Source/Programs/UnrealCloudDDC/Jupiter/Implementation/Objects/ScyllaReferencesStore.cs
@@ -533,16 +533,11 @@ namespace Jupiter.Implementation
 
 			Task? lastAccessDeleteTask = _mapper.DeleteAsync<ScyllaObjectLastAccess>("WHERE namespace=? AND bucket=? AND name=?", ns.ToString(), bucket.ToString(), key.ToString());
 
-			AppliedInfo<ScyllaObject> info = await _mapper.DeleteIfAsync<ScyllaObject>("WHERE namespace=? AND bucket=? AND name=? IF EXISTS", ns.ToString(), bucket.ToString(), key.ToString());
+			await _mapper.DeleteAsync<ScyllaObject>("WHERE namespace=? AND bucket=? AND name=?", ns.ToString(), bucket.ToString(), key.ToString());
 
 			await lastAccessDeleteTask;
 
-			if (info.Applied)
-			{
-				return true;
-			}
-
-			return false;
+			return true;
 		}
 
 		public async Task<long> DropNamespaceAsync(NamespaceId ns)
