Unreal Cloud DDC - Increased number of requests allowed in a keepalive connection - seems to reduce connection issues. Also enabled proxy buffering which seems to help a bit with http2 multiplexing and is generally recommended by nginx. Lastly I made common error messages return json rather then a html page which makes it look a bit better.

diff --git a/Engine/Source/Programs/UnrealCloudDDC/Helm/UnrealCloudDDC/templates/configmap-nginx.yaml b/Engine/Source/Programs/UnrealCloudDDC/Helm/UnrealCloudDDC/templates/configmap-nginx.yaml
--- a/Engine/Source/Programs/UnrealCloudDDC/Helm/UnrealCloudDDC/templates/configmap-nginx.yaml
+++ b/Engine/Source/Programs/UnrealCloudDDC/Helm/UnrealCloudDDC/templates/configmap-nginx.yaml
@@ -26,6 +26,9 @@ data:
       sendfile on;
       tcp_nopush on;
 
+      # increases the number of requests allowed per keepalive connection, reduces connection churn at a cost of some memory usage
+      keepalive_requests 10000;
+
       upstream jupiter_backend {
         zone upstreams 64K;
         server {{ if .Values.nginx.useDomainSockets -}} unix:/nginx/jupiter-http.sock {{- else -}} localhost:80 {{- end }} max_fails=1 fail_timeout=2s;
@@ -37,7 +40,7 @@ data:
         listen 81 backlog={{ .Values.nginx.backlog }};
         server_name public_http1;
 
-        proxy_buffering off;
+        proxy_buffering on;
         proxy_request_buffering off;
         proxy_http_version 1.1;
         
@@ -69,10 +72,10 @@ data:
             return 404;
         }
 
-        error_page 404 @error_page;
+        error_page 404 502 504 @error_page;
         location @error_page {
-          default_type "";
-          return 200; # error_code is ignored
+          default_type "application/json";
+          return 200 '{"error_code": $status}'; # return_code is ignored
         }
       }
 
@@ -81,7 +84,7 @@ data:
         listen 8009 backlog={{ .Values.nginx.backlog }};
         server_name corp_http1;
         
-        proxy_buffering off;
+        proxy_buffering on;
         proxy_request_buffering off;
         proxy_http_version 1.1;
         
@@ -113,19 +116,20 @@ data:
             return 404;
         }
 
-        error_page 404 @error_page;
+        error_page 404 502 504 @error_page;
         location @error_page {
-          default_type "";
-          return 200; # error_code is ignored
+          default_type "application/json";
+          return 200 '{"error_code": $status}'; # return_code is ignored
         }
       }
 
       # public http2 port
       server {
-        listen 8091 backlog={{ .Values.nginx.backlog }} http2;
+        listen 8091 backlog={{ .Values.nginx.backlog }};
         server_name public_http2;
+        http2 on;
 
-        proxy_buffering off;
+        proxy_buffering on;
         proxy_request_buffering off;
         proxy_http_version 1.1;
         
@@ -157,19 +161,20 @@ data:
             return 404;
         }
 
-        error_page 404 @error_page;
+        error_page 404 502 504 @error_page;
         location @error_page {
-          default_type "";
-          return 200; # error_code is ignored
+          default_type "application/json";
+          return 200 '{"error_code": $status}'; # return_code is ignored
         }
       }
 
       # corp port http2
       server {
-        listen 8092 backlog={{ .Values.nginx.backlog }} http2;
+        listen 8092 backlog={{ .Values.nginx.backlog }};
         server_name corp_http2;
+        http2 on;
 
-        proxy_buffering off;
+        proxy_buffering on;
         proxy_request_buffering off;
         proxy_http_version 1.1;
         
@@ -201,10 +206,10 @@ data:
             return 404;
         }
 
-        error_page 404 @error_page;
+        error_page 404 502 504 @error_page;
         location @error_page {
-          default_type "";
-          return 200; # error_code is ignored
+          default_type "application/json";
+          return 200 '{"error_code": $status}'; # return_code is ignored
         }
       }
 
