Unreal Cloud DDC - Updated nginx configuration for docker compose setup to be more similar to how we normally run it (enabled domain sockets, reconfigured keepalive connections, enabled http2 port and enabled proxy buffering)

diff --git a/Engine/Source/Programs/UnrealCloudDDC/Composes/docker-compose.yml b/Engine/Source/Programs/UnrealCloudDDC/Composes/docker-compose.yml
--- a/Engine/Source/Programs/UnrealCloudDDC/Composes/docker-compose.yml
+++ b/Engine/Source/Programs/UnrealCloudDDC/Composes/docker-compose.yml
@@ -6,6 +6,7 @@ services:
       dockerfile: nginx_dockerfile
     ports:
     - "80:80" 
+    - "8008:8008" 
     networks:
     - network_jupiter
 
@@ -13,6 +14,7 @@ services:
       - unreal_cloud_ddc
     volumes:
      - jupiter-storage-volume:/data/blob-store
+     - nginx-shared-volume:/nginx
   scylla:
     logging:
       driver: none
@@ -34,6 +36,9 @@ services:
     - network_jupiter
     environment:
     - Jupiter__CurrentSite=docker-compose
+    - Jupiter__UseDomainSockets=true
+    - Jupiter__DomainSocketsRoot=/nginx
+    - Jupiter__ChmodDomainSockets=true
     # - Serilog__MinimumLevel__Default=Debug # Enable this to get debug logging, can be quite a lot
     - CORECLR_ENABLE_PROFILING=0
     - Auth__Enabled=false
@@ -50,9 +55,11 @@ services:
       com.datadoghq.ad.logs: '[{"source":"csharp", "service": "jupiter"}]'
     volumes:
      - jupiter-storage-volume:/data/blob-store
+     - nginx-shared-volume:/nginx
 volumes:
   scylla-volume-1:
   jupiter-storage-volume:
+  nginx-shared-volume:
 networks:
   network_jupiter:
     driver: bridge
