Unreal Cloud DDC - More fixes for consistency in api response between backends when deleting a ref that does not exist.

diff --git a/Engine/Source/Programs/UnrealCloudDDC/Jupiter/Controllers/ReferencesController.cs b/Engine/Source/Programs/UnrealCloudDDC/Jupiter/Controllers/ReferencesController.cs
--- a/Engine/Source/Programs/UnrealCloudDDC/Jupiter/Controllers/ReferencesController.cs
+++ b/Engine/Source/Programs/UnrealCloudDDC/Jupiter/Controllers/ReferencesController.cs
@@ -1121,7 +1121,7 @@ namespace Jupiter.Controllers
 			}
 			catch (RefNotFoundException)
 			{
-				return NotFound(new ProblemDetails { Title = $"Object {key} in bucket {bucket} and namespace {ns} did not exist" });
+				return Ok(new RefDeletedResponse(0));
 			}
 		}
 	}
