Unreal Cloud DDC - Upgraded nginx version used to 1.25.4
Should fix startup issue as the http2 directive is new.

diff --git a/Engine/Source/Programs/UnrealCloudDDC/Helm/UnrealCloudDDC/templates/deployment.yaml b/Engine/Source/Programs/UnrealCloudDDC/Helm/UnrealCloudDDC/templates/deployment.yaml
--- a/Engine/Source/Programs/UnrealCloudDDC/Helm/UnrealCloudDDC/templates/deployment.yaml
+++ b/Engine/Source/Programs/UnrealCloudDDC/Helm/UnrealCloudDDC/templates/deployment.yaml
@@ -163,7 +163,7 @@ spec:
     {{- end }}
     {{- if .Values.nginx.enabled }}
       - name: nginx
-        image: nginx:1.23.3
+        image: nginx:1.25.4
         {{- if .Values.nginx.env }}
         env:
         {{- include "common.tplvalues.render" ( dict "value" .Values.nginx.env "context" $ ) | nindent 8 }}
