Unreal Cloud DDC - Delay checking blob store for uncompressed hash when resolving content id until that is actually used. Makes the content id from uncompressed hash slower (but we never use that in practise) but avoids unnesscary calls to blob resolving (eg reduces s3 operations)

diff --git a/Engine/Source/Programs/UnrealCloudDDC/Jupiter/Implementation/ContentId/ScyllaContentIdStore.cs b/Engine/Source/Programs/UnrealCloudDDC/Jupiter/Implementation/ContentId/ScyllaContentIdStore.cs
--- a/Engine/Source/Programs/UnrealCloudDDC/Jupiter/Implementation/ContentId/ScyllaContentIdStore.cs
+++ b/Engine/Source/Programs/UnrealCloudDDC/Jupiter/Implementation/ContentId/ScyllaContentIdStore.cs
@@ -46,11 +46,6 @@ namespace Jupiter.Implementation
 			using TelemetrySpan scope = _tracer.BuildScyllaSpan("ScyllaContentIdStore.ResolveContentId").SetAttribute("resource.name", contentId.ToString());
 
 			BlobId contentIdBlob = contentId.AsBlobIdentifier();
-			Task<bool>? blobStoreExistsTask = null;
-			if (!mustBeContentId)
-			{
-				blobStoreExistsTask = _blobStore.ExistsAsync(ns, contentIdBlob);
-			}
 
 			{
 				using TelemetrySpan contentIdFetchScope = _tracer.BuildScyllaSpan("ScyllaContentIdStore.FetchContentId").SetAttribute("resource.name", contentId.ToString());
@@ -109,7 +104,7 @@ namespace Jupiter.Implementation
 			if (!mustBeContentId)
 			{
 				// if no content id is found, but we have a blob that matches the content id (so a unchunked and uncompressed version of the data) we use that instead
-				bool contentIdBlobExists = await blobStoreExistsTask!;
+				bool contentIdBlobExists = await _blobStore.ExistsAsync(ns, contentIdBlob)!;
 
 				if (contentIdBlobExists)
 				{
