Unreal Cloud DDC - Added ability to fetch metadata about blobs without trigger caching of the blob in the filesystem store.
Fixes issue with ref puts taking much longer for large refs (specifically oplog uploads).

diff --git a/Engine/Source/Programs/UnrealCloudDDC/Jupiter/Implementation/Blob/AmazonS3Store.cs b/Engine/Source/Programs/UnrealCloudDDC/Jupiter/Implementation/Blob/AmazonS3Store.cs
--- a/Engine/Source/Programs/UnrealCloudDDC/Jupiter/Implementation/Blob/AmazonS3Store.cs
+++ b/Engine/Source/Programs/UnrealCloudDDC/Jupiter/Implementation/Blob/AmazonS3Store.cs
@@ -56,6 +56,17 @@ namespace Jupiter.Implementation
 			return uri;
 		}
 
+		public async Task<BlobMetadata> GetObjectMetadataAsync(NamespaceId ns, BlobId identifier)
+		{
+			BlobMetadata? result = await GetBackend(ns).GetMetadataAsync(identifier.AsS3Key());
+			if (result == null)
+			{
+				throw new BlobNotFoundException(ns, identifier);
+			}
+
+			return result;
+		}
+
 		public async Task<Uri?> PutObjectWithRedirectAsync(NamespaceId ns, BlobId identifier)
 		{
 			Uri? uri = await GetBackend(ns).GetWriteRedirectAsync(identifier.AsS3Key());
@@ -461,6 +472,17 @@ namespace Jupiter.Implementation
 				return null;
 			}
 		}
+
+		public async Task<BlobMetadata?> GetMetadataAsync(string path)
+		{
+			GetObjectAttributesResponse? metadata = await _amazonS3.GetObjectAttributesAsync(new GetObjectAttributesRequest {BucketName = _bucketName, Key = path});
+			if (metadata == null)
+			{
+				return null;
+			}
+
+			return new BlobMetadata(metadata.ObjectSize, metadata.LastModified);
+		}
 	}
 
 	public static class BlobIdentifierExtensions
diff --git a/Engine/Source/Programs/UnrealCloudDDC/Jupiter/Implementation/Blob/AzureBlobStore.cs b/Engine/Source/Programs/UnrealCloudDDC/Jupiter/Implementation/Blob/AzureBlobStore.cs
--- a/Engine/Source/Programs/UnrealCloudDDC/Jupiter/Implementation/Blob/AzureBlobStore.cs
+++ b/Engine/Source/Programs/UnrealCloudDDC/Jupiter/Implementation/Blob/AzureBlobStore.cs
@@ -78,6 +78,18 @@ namespace Jupiter.Implementation
 			return redirectUri;
 		}
 
+		public async Task<BlobMetadata> GetObjectMetadataAsync(NamespaceId ns, BlobId blobId)
+		{
+			BlobMetadata? metadata = await GetBackend(ns).GetMetadataAsync(GetPath(blobId), CancellationToken.None);
+
+			if (metadata == null)
+			{
+				throw new BlobNotFoundException(ns, blobId);
+			}
+
+			return metadata;
+		}
+
 		public async Task<Uri?> PutObjectWithRedirectAsync(NamespaceId ns, BlobId identifier)
 		{
 			Uri? redirectUri = await GetBackend(ns).GetWriteRedirectAsync(GetPath(identifier));
@@ -323,5 +335,22 @@ namespace Jupiter.Implementation
 				throw;
 			}
 		}
+
+		public async Task<BlobMetadata?> GetMetadataAsync(string path, CancellationToken cancellationToken)
+		{
+			if (!await _blobContainer.ExistsAsync(cancellationToken))
+			{
+				throw new InvalidOperationException($"Container {_blobContainer.Name} did not exist");
+			}
+
+			BlobClient blob = _blobContainer.GetBlobClient(path);
+			Response<BlobProperties>? response = await blob.GetPropertiesAsync(cancellationToken: cancellationToken);
+			if (response == null)
+			{
+				return null;
+			}
+
+			return new BlobMetadata(response.Value.ContentLength, response.Value.CreatedOn.DateTime);
+		}
 	}
 }
diff --git a/Engine/Source/Programs/UnrealCloudDDC/Jupiter/Implementation/Blob/BlobService.cs b/Engine/Source/Programs/UnrealCloudDDC/Jupiter/Implementation/Blob/BlobService.cs
--- a/Engine/Source/Programs/UnrealCloudDDC/Jupiter/Implementation/Blob/BlobService.cs
+++ b/Engine/Source/Programs/UnrealCloudDDC/Jupiter/Implementation/Blob/BlobService.cs
@@ -513,6 +513,55 @@ public class BlobService : IBlobService
 		return redirectUri;
 	}
 
+	public async Task<BlobMetadata> GetObjectMetadataAsync(NamespaceId ns, BlobId blobId)
+	{
+		bool seenBlobNotFound = false;
+		bool seenNamespaceNotFound = false;
+		IServerTiming? serverTiming = _httpContextAccessor.HttpContext?.RequestServices.GetService<IServerTiming>();
+
+		foreach (IBlobStore store in _blobStores)
+		{
+			string storeName = store.GetType().Name;
+
+			using TelemetrySpan scope = _tracer.StartActiveSpan("HierarchicalStore.GetObjectMetadata")
+					.SetAttribute("operation.name", "HierarchicalStore.GetObjectMetadata")
+					.SetAttribute("resource.name", blobId.ToString())
+					.SetAttribute("BlobStore", storeName)
+					.SetAttribute("ObjectFound", false.ToString())
+				;
+
+			using ServerTimingMetricScoped? serverTimingScope = serverTiming?.CreateServerTimingMetricScope($"blob.get-metadata.{storeName}", $"Blob GET Metadata from: '{storeName}'");
+
+			try
+			{
+				BlobMetadata metadata = await store.GetObjectMetadataAsync(ns, blobId);
+				scope.SetAttribute("ObjectFound", true.ToString());
+				return metadata;
+			}
+			catch (BlobNotFoundException)
+			{
+				seenBlobNotFound = true;
+			}
+			catch (NamespaceNotFoundException)
+			{
+				seenNamespaceNotFound = true;
+			}
+		}
+
+		if (seenBlobNotFound)
+		{
+			throw new BlobNotFoundException(ns, blobId);
+		}
+
+		if (seenNamespaceNotFound)
+		{
+			throw new NamespaceNotFoundException(ns);
+		}
+
+		// the only way we get here is we can not find the blob in any store, thus it does not exist (should have triggered the blob not found above)
+		throw new BlobNotFoundException(ns, blobId);
+	}
+
 	public async Task<BlobContents> ReplicateObjectAsync(NamespaceId ns, BlobId blob, bool force = false)
 	{
 		if (!force && !ShouldFetchBlobOnDemand(ns))
diff --git a/Engine/Source/Programs/UnrealCloudDDC/Jupiter/Implementation/Blob/FileSystemStore.cs b/Engine/Source/Programs/UnrealCloudDDC/Jupiter/Implementation/Blob/FileSystemStore.cs
--- a/Engine/Source/Programs/UnrealCloudDDC/Jupiter/Implementation/Blob/FileSystemStore.cs
+++ b/Engine/Source/Programs/UnrealCloudDDC/Jupiter/Implementation/Blob/FileSystemStore.cs
@@ -69,6 +69,12 @@ namespace Jupiter.Implementation
 			return Task.FromResult<Uri?>(null);
 		}
 
+		public Task<BlobMetadata> GetObjectMetadataAsync(NamespaceId ns, BlobId blobId)
+		{
+			string path = GetFilesystemPath(blobId);
+			return GetBackend(ns).GetMetadata(path);
+		}
+
 		public Task<Uri?> PutObjectWithRedirectAsync(NamespaceId ns, BlobId identifier)
 		{
 			// not supported
@@ -438,5 +444,11 @@ namespace Jupiter.Implementation
 				yield return (path, file.LastWriteTime);
 			}
 		}
+
+		public Task<BlobMetadata> GetMetadata(string path)
+		{
+			FileInfo fileInfo = GetFilesystemPath(path);
+			return Task.FromResult(new BlobMetadata(fileInfo.Length, fileInfo.CreationTime));
+		}
 	}
 }
diff --git a/Engine/Source/Programs/UnrealCloudDDC/Jupiter/Implementation/Blob/IBlobService.cs b/Engine/Source/Programs/UnrealCloudDDC/Jupiter/Implementation/Blob/IBlobService.cs
--- a/Engine/Source/Programs/UnrealCloudDDC/Jupiter/Implementation/Blob/IBlobService.cs
+++ b/Engine/Source/Programs/UnrealCloudDDC/Jupiter/Implementation/Blob/IBlobService.cs
@@ -25,6 +25,8 @@ public interface IBlobService
 	
 	Task<Uri?> GetObjectWithRedirectAsync(NamespaceId ns, BlobId blobIdentifier, List<string>? storageLayers = null);
 
+	Task<BlobMetadata> GetObjectMetadataAsync(NamespaceId ns, BlobId blobId);
+
 	Task<BlobContents> ReplicateObjectAsync(NamespaceId ns, BlobId blob, bool force = false);
 
 	Task<bool> ExistsAsync(NamespaceId ns, BlobId blob, List<string>? storageLayers = null);
@@ -51,6 +53,18 @@ public interface IBlobService
 	bool ShouldFetchBlobOnDemand(NamespaceId ns);
 }
 
+public class BlobMetadata
+{
+	public BlobMetadata(long length, DateTime creationTime)
+	{
+		Length = length;
+		CreationTime = creationTime;
+	}
+
+	public long Length { get; set; }
+	public DateTime CreationTime { get; set; }
+}
+
 public static class BlobServiceExtensions
 {
 	public static async Task<ContentId> PutCompressedObjectAsync(this IBlobService blobService, NamespaceId ns, IBufferedPayload payload, ContentId? id, IServiceProvider provider)
diff --git a/Engine/Source/Programs/UnrealCloudDDC/Jupiter/Implementation/Blob/IBlobStore.cs b/Engine/Source/Programs/UnrealCloudDDC/Jupiter/Implementation/Blob/IBlobStore.cs
--- a/Engine/Source/Programs/UnrealCloudDDC/Jupiter/Implementation/Blob/IBlobStore.cs
+++ b/Engine/Source/Programs/UnrealCloudDDC/Jupiter/Implementation/Blob/IBlobStore.cs
@@ -33,6 +33,7 @@ namespace Jupiter.Implementation
 		IAsyncEnumerable<(BlobId,DateTime)> ListObjectsAsync(NamespaceId ns);
 		Task<Uri?> PutObjectWithRedirectAsync(NamespaceId ns, BlobId identifier);
 		Task<Uri?> GetObjectByRedirectAsync(NamespaceId ns, BlobId blob);
+		Task<BlobMetadata> GetObjectMetadataAsync(NamespaceId ns, BlobId blobId);
 	}
 
 	public interface IStorageBackend
diff --git a/Engine/Source/Programs/UnrealCloudDDC/Jupiter/Implementation/Blob/MemoryBlobStore.cs b/Engine/Source/Programs/UnrealCloudDDC/Jupiter/Implementation/Blob/MemoryBlobStore.cs
--- a/Engine/Source/Programs/UnrealCloudDDC/Jupiter/Implementation/Blob/MemoryBlobStore.cs
+++ b/Engine/Source/Programs/UnrealCloudDDC/Jupiter/Implementation/Blob/MemoryBlobStore.cs
@@ -48,6 +48,17 @@ namespace Jupiter.Implementation
 			return Task.FromResult<Uri?>(null);
 		}
 
+		public async Task<BlobMetadata> GetObjectMetadataAsync(NamespaceId ns, BlobId identifier)
+		{
+			BlobMetadata? metadata = await GetBackend(ns).GetObjectMetadataAsync(GetPath(identifier));
+			if (metadata == null)
+			{
+				throw new BlobNotFoundException(ns, identifier);
+			}
+
+			return metadata;
+		}
+
 		public Task<Uri?> PutObjectWithRedirectAsync(NamespaceId ns, BlobId identifier)
 		{
 			// not supported
@@ -195,5 +206,16 @@ namespace Jupiter.Implementation
 				blobContainer.LastModified = modifiedTime;
 			}
 		}
+
+		public Task<BlobMetadata?> GetObjectMetadataAsync(string path)
+		{
+			if (!_blobs.TryGetValue(path, value: out BlobContainer? blobContainer))
+			{
+				return Task.FromResult<BlobMetadata?>(null);
+			}
+
+			byte[] content = blobContainer.Contents;
+			return Task.FromResult<BlobMetadata?>(new BlobMetadata(content.LongLength, blobContainer.LastModified));
+		}
 	}
 }
diff --git a/Engine/Source/Programs/UnrealCloudDDC/Jupiter/Implementation/Blob/PeerBlobStore.cs b/Engine/Source/Programs/UnrealCloudDDC/Jupiter/Implementation/Blob/PeerBlobStore.cs
--- a/Engine/Source/Programs/UnrealCloudDDC/Jupiter/Implementation/Blob/PeerBlobStore.cs
+++ b/Engine/Source/Programs/UnrealCloudDDC/Jupiter/Implementation/Blob/PeerBlobStore.cs
@@ -171,6 +171,12 @@ namespace Jupiter.Implementation
 			return Task.FromResult<Uri?>(null);
 		}
 
+		public Task<BlobMetadata> GetObjectMetadataAsync(NamespaceId ns, BlobId blobId)
+		{
+			// Do not call into other instances to determine metadata as we lack a endpoint for it
+			throw new BlobNotFoundException(ns, blobId);
+		}
+
 		public Task<Uri?> PutObjectWithRedirectAsync(NamespaceId ns, BlobId identifier)
 		{
 			// not supported
diff --git a/Engine/Source/Programs/UnrealCloudDDC/Jupiter/Implementation/Blob/RelayBlobStore.cs b/Engine/Source/Programs/UnrealCloudDDC/Jupiter/Implementation/Blob/RelayBlobStore.cs
--- a/Engine/Source/Programs/UnrealCloudDDC/Jupiter/Implementation/Blob/RelayBlobStore.cs
+++ b/Engine/Source/Programs/UnrealCloudDDC/Jupiter/Implementation/Blob/RelayBlobStore.cs
@@ -25,6 +25,12 @@ namespace Jupiter.Implementation
 			return Task.FromResult<Uri?>(null);
 		}
 
+		public Task<BlobMetadata> GetObjectMetadataAsync(NamespaceId ns, BlobId blobId)
+		{
+			// do not call into other instances to find metadata as we lack a endpoint for that
+			throw new BlobNotFoundException(ns, blobId);
+		}
+
 		public Task<Uri?> PutObjectWithRedirectAsync(NamespaceId ns, BlobId identifier)
 		{
 			// TODO: It could be useful to support relaying the presigned url
diff --git a/Engine/Source/Programs/UnrealCloudDDC/Jupiter/Implementation/RefService.cs b/Engine/Source/Programs/UnrealCloudDDC/Jupiter/Implementation/RefService.cs
--- a/Engine/Source/Programs/UnrealCloudDDC/Jupiter/Implementation/RefService.cs
+++ b/Engine/Source/Programs/UnrealCloudDDC/Jupiter/Implementation/RefService.cs
@@ -225,7 +225,7 @@ namespace Jupiter.Implementation
 								// if a blob is missing its not a error, the finalize will report this as missing and it will be uploaded and finalize ran again
 								try
 								{
-									BlobContents result = await _blobService.GetObjectAsync(ns, blobId);
+									BlobMetadata result = await _blobService.GetObjectMetadataAsync(ns, blobId);
 									await _blobIndex.AddBlobToBucketListAsync(ns, bucket, key, blobId, result.Length);
 								}
 								catch (BlobNotFoundException)
