Unreal Cloud DDC - If a do not find a local blob when looking in the filesystem store we want to raise a blobnotfound exception not a filenotfound exception

diff --git a/Engine/Source/Programs/UnrealCloudDDC/Jupiter/Implementation/Blob/FileSystemStore.cs b/Engine/Source/Programs/UnrealCloudDDC/Jupiter/Implementation/Blob/FileSystemStore.cs
--- a/Engine/Source/Programs/UnrealCloudDDC/Jupiter/Implementation/Blob/FileSystemStore.cs
+++ b/Engine/Source/Programs/UnrealCloudDDC/Jupiter/Implementation/Blob/FileSystemStore.cs
@@ -71,8 +71,15 @@ namespace Jupiter.Implementation
 
 		public Task<BlobMetadata> GetObjectMetadataAsync(NamespaceId ns, BlobId blobId)
 		{
-			string path = GetFilesystemPath(blobId);
-			return GetBackend(ns).GetMetadata(path);
+			try
+			{
+				string path = GetFilesystemPath(blobId);
+				return GetBackend(ns).GetMetadata(path);
+			}
+			catch (FileNotFoundException)
+			{
+				throw new BlobNotFoundException(ns, blobId);
+			}
 		}
 
 		public Task<Uri?> PutObjectWithRedirectAsync(NamespaceId ns, BlobId identifier)
