Unreal Cloud DDC - Fixed issue with s3 not returning metadata we want (instead it raised a exception) when requesting object metadata

diff --git a/Engine/Source/Programs/UnrealCloudDDC/Jupiter/Implementation/Blob/AmazonS3Store.cs b/Engine/Source/Programs/UnrealCloudDDC/Jupiter/Implementation/Blob/AmazonS3Store.cs
--- a/Engine/Source/Programs/UnrealCloudDDC/Jupiter/Implementation/Blob/AmazonS3Store.cs
+++ b/Engine/Source/Programs/UnrealCloudDDC/Jupiter/Implementation/Blob/AmazonS3Store.cs
@@ -475,7 +475,10 @@ namespace Jupiter.Implementation
 
 		public async Task<BlobMetadata?> GetMetadataAsync(string path)
 		{
-			GetObjectAttributesResponse? metadata = await _amazonS3.GetObjectAttributesAsync(new GetObjectAttributesRequest {BucketName = _bucketName, Key = path});
+			GetObjectAttributesResponse? metadata = await _amazonS3.GetObjectAttributesAsync(new GetObjectAttributesRequest {BucketName = _bucketName, Key = path, ObjectAttributes = new List<ObjectAttributes>()
+			{
+				ObjectAttributes.ObjectSize
+			}});
 			if (metadata == null)
 			{
 				return null;
