Unreal Cloud DDC - Delay waiting for bucket stats task for as long as possible allowing more work to run in parallel.

Also catch S3 issue if the blob does not exist when requesting metadata about it.

diff --git a/Engine/Source/Programs/UnrealCloudDDC/Jupiter/Implementation/Blob/AmazonS3Store.cs b/Engine/Source/Programs/UnrealCloudDDC/Jupiter/Implementation/Blob/AmazonS3Store.cs
--- a/Engine/Source/Programs/UnrealCloudDDC/Jupiter/Implementation/Blob/AmazonS3Store.cs
+++ b/Engine/Source/Programs/UnrealCloudDDC/Jupiter/Implementation/Blob/AmazonS3Store.cs
@@ -475,16 +475,34 @@ namespace Jupiter.Implementation
 
 		public async Task<BlobMetadata?> GetMetadataAsync(string path)
 		{
-			GetObjectAttributesResponse? metadata = await _amazonS3.GetObjectAttributesAsync(new GetObjectAttributesRequest {BucketName = _bucketName, Key = path, ObjectAttributes = new List<ObjectAttributes>()
-			{
-				ObjectAttributes.ObjectSize
-			}});
-			if (metadata == null)
+			try
 			{
-				return null;
+				GetObjectAttributesResponse? metadata = await _amazonS3.GetObjectAttributesAsync(new GetObjectAttributesRequest
+				{
+					BucketName = _bucketName,
+					Key = path,
+					ObjectAttributes = new List<ObjectAttributes>()
+					{
+						ObjectAttributes.ObjectSize
+					}
+				});
+				if (metadata == null)
+				{
+					return null;
+				}
+
+				return new BlobMetadata(metadata.ObjectSize, metadata.LastModified);
 			}
+			catch (AmazonS3Exception e)
+			{
+				// if the object does not exist we get a not found status code
+				if (e.StatusCode == HttpStatusCode.NotFound)
+				{
+					return null;
+				}
 
-			return new BlobMetadata(metadata.ObjectSize, metadata.LastModified);
+				throw;
+			}
 		}
 	}
 
diff --git a/Engine/Source/Programs/UnrealCloudDDC/Jupiter/Implementation/RefService.cs b/Engine/Source/Programs/UnrealCloudDDC/Jupiter/Implementation/RefService.cs
--- a/Engine/Source/Programs/UnrealCloudDDC/Jupiter/Implementation/RefService.cs
+++ b/Engine/Source/Programs/UnrealCloudDDC/Jupiter/Implementation/RefService.cs
@@ -205,12 +205,13 @@ namespace Jupiter.Implementation
 				: Task.CompletedTask;
 			ContentId[] missingReferences = Array.Empty<ContentId>();
 			BlobId[] missingBlobs = Array.Empty<BlobId>();
+
+			List<Task> addToBucketTasks = new List<Task>();
+			List<Task> addRefMappingTasks = new List<Task>();
+
 			bool hasReferences = HasAttachments(payload);
 			if (hasReferences)
 			{
-				List<Task> addToBucketTasks = new List<Task>();
-				List<Task> addRefMappingTasks = new List<Task>();
-
 				using TelemetrySpan _ = _tracer.StartActiveSpan("ObjectService.ResolveReferences").SetAttribute("operation.name", "ObjectService.ResolveReferences");
 				try
 				{
@@ -236,8 +237,6 @@ namespace Jupiter.Implementation
 
 						addRefMappingTasks.Add(_blobIndex.AddRefToBlobsAsync(ns, bucket, key, new BlobId[] {blobId}));
 					}
-
-					await Task.WhenAll(addRefToBlobsTask, addToBucketListTask, Task.WhenAll(addToBucketTasks), Task.WhenAll(addRefMappingTasks));
 				}
 				catch (PartialReferenceResolveException e)
 				{
@@ -248,10 +247,8 @@ namespace Jupiter.Implementation
 					missingBlobs = e.MissingBlobs.ToArray();
 				}
 			}
-			else
-			{
-				await Task.WhenAll(addRefToBlobsTask, addToBucketListTask);
-			}
+
+			await Task.WhenAll(addRefToBlobsTask, addToBucketListTask);
 
 			if (missingReferences.Length == 0 && missingBlobs.Length == 0)
 			{
@@ -259,6 +256,8 @@ namespace Jupiter.Implementation
 				await _replicationLog.InsertAddEventAsync(ns, bucket, key, blobHash);
 			}
 
+			await Task.WhenAll(Task.WhenAll(addToBucketTasks), Task.WhenAll(addRefMappingTasks));
+
 			return (missingReferences, missingBlobs);
 		}
 
