Unreal Cloud DDC - Use a yaml file for configuration for docker compose - to make it easier to understand how to change configuration when comparing to the helm docs. Also provided a example of setting a namespace acl even if it not used by default.

diff --git a/Engine/Source/Programs/UnrealCloudDDC/Composes/appsettings.Local.yaml b/Engine/Source/Programs/UnrealCloudDDC/Composes/appsettings.Local.yaml
new file mode 100644
--- /dev/null
+++ b/Engine/Source/Programs/UnrealCloudDDC/Composes/appsettings.Local.yaml
@@ -0,0 +1,38 @@
+Serilog:
+  MinimumLevel:
+    Default:
+      Warning # Toogle this to Debug to get debug logging, can be quite a lot
+
+Auth:
+  Enabled: false
+Jupiter:
+  CurrentSite: docker-compose
+  UseDomainSockets: true
+  DomainSocketsRoot: /nginx
+  ChmodDomainSockets: true
+Replication:
+  CurrentSite: docker-compose
+Ngunx:
+  UseNginxRedirect: true
+UnrealCloudDDC:
+  storageImplementations:
+  - Filesystem
+Filesystem:
+  RootDir: /data/blob-store
+Scylla:
+  ContactPoints:
+  - scylla
+  LocalDatacenterName: datacenter1
+  LocalKeyspaceSuffix: dc1
+
+Namespaces:
+  Policies:
+    # Define a single namespace that is open to anyone with a valid credential to use
+    # As auth is disabled in this config this doesn't really matter and is provided mostly for reference
+    test-namespace:
+    - acls:
+      actions:
+      - ReadObject
+      - WriteObject
+      claims:
+      - "*"
\ No newline at end of file
diff --git a/Engine/Source/Programs/UnrealCloudDDC/Composes/docker-compose.yml b/Engine/Source/Programs/UnrealCloudDDC/Composes/docker-compose.yml
--- a/Engine/Source/Programs/UnrealCloudDDC/Composes/docker-compose.yml
+++ b/Engine/Source/Programs/UnrealCloudDDC/Composes/docker-compose.yml
@@ -34,21 +34,11 @@ services:
     - "8080:8080" 
     networks:
     - network_jupiter
+    configs:
+    - source: cloud_ddc_config
+      target: /config/appsettings.Local.yaml
     environment:
-    - Jupiter__CurrentSite=docker-compose
-    - Jupiter__UseDomainSockets=true
-    - Jupiter__DomainSocketsRoot=/nginx
-    - Jupiter__ChmodDomainSockets=true
-    # - Serilog__MinimumLevel__Default=Debug # Enable this to get debug logging, can be quite a lot
     - CORECLR_ENABLE_PROFILING=0
-    - Auth__Enabled=false
-    - Replication__CurrentSite=docker-compose
-    - Scylla__ContactPoints__0=scylla
-    - Scylla__LocalDatacenterName=datacenter1
-    - Scylla__LocalKeyspaceSuffix=dc1
-    - Filesystem__RootDir=/data/blob-store
-    - Nginx__UseNginxRedirect=true
-    - UnrealCloudDDC__storageImplementations__0=Filesystem
     links:
     - scylla
     labels:
@@ -63,3 +53,6 @@ volumes:
 networks:
   network_jupiter:
     driver: bridge
+configs:
+  cloud_ddc_config:
+    file: appsettings.Local.yaml
