Unreal Cloud DDC - Corrected example configuration for namespace, it was not correctly indicating that ACLs is an array

diff --git a/Engine/Source/Programs/UnrealCloudDDC/Helm/UnrealCloudDDC/example-values-base.yaml b/Engine/Source/Programs/UnrealCloudDDC/Helm/UnrealCloudDDC/example-values-base.yaml
--- a/Engine/Source/Programs/UnrealCloudDDC/Helm/UnrealCloudDDC/example-values-base.yaml
+++ b/Engine/Source/Programs/UnrealCloudDDC/Helm/UnrealCloudDDC/example-values-base.yaml
@@ -64,7 +64,7 @@ global:
       # Define a single namespace that is open to anyone with a valid credential to use
       test-namespace:
         acls:
-          actions:
+        - actions:
           - ReadObject
           - WriteObject
           claims:
