Unreal Cloud DDC - Prepared release 1.0.0

diff --git a/Engine/Source/Programs/UnrealCloudDDC/Helm/UnrealCloudDDC/version.yaml b/Engine/Source/Programs/UnrealCloudDDC/Helm/UnrealCloudDDC/version.yaml
--- a/Engine/Source/Programs/UnrealCloudDDC/Helm/UnrealCloudDDC/version.yaml
+++ b/Engine/Source/Programs/UnrealCloudDDC/Helm/UnrealCloudDDC/version.yaml
@@ -1 +1 @@
-version: 0.6.0
\ No newline at end of file
+version: 1.0.0
\ No newline at end of file
diff --git a/Engine/Source/Programs/UnrealCloudDDC/changelog.md b/Engine/Source/Programs/UnrealCloudDDC/changelog.md
--- a/Engine/Source/Programs/UnrealCloudDDC/changelog.md
+++ b/Engine/Source/Programs/UnrealCloudDDC/changelog.md
@@ -1,5 +1,15 @@
 # Unreleased
 
+# 1.0.0
+* .NET 8 Upgrade.
+* Fixes for very large payloads (2GB+)
+* Bug fixes for replication, mostly fixes to improve behavior when replication has fallen behind.
+* Reduced GC pressure when processing a lot of new blobs that could cause very long stalls randomly (up to 1s response times).
+* Speed up GC of refs when using Scylla.
+* Tweaks to nginx configuration when using it as a reverse proxy.
+* Reference store consistency check added, enabled via `ConsistencyCheck:EnableRefStoreChecks`. Can help repair inconsitencies between tables in Scylla which could happen due to bugs in earlier versions (resulting in data that is never garbage collected).
+* Endpoints for serving symbol data that is compatible with MS Symbol server http api (but still requires auth which is not supported by Visual Studio out of the box). Experimental feature.
+
 # 0.6.0
 * Added option to track per bucket stats `EnableBucketStatsTracking`, this is still WIP.
 * Added option to tweak nginx keep alive connections and increased it, only applies if using the nginx proxy.
diff --git a/Engine/Source/Programs/UnrealCloudDDC/version.yaml b/Engine/Source/Programs/UnrealCloudDDC/version.yaml
--- a/Engine/Source/Programs/UnrealCloudDDC/version.yaml
+++ b/Engine/Source/Programs/UnrealCloudDDC/version.yaml
@@ -1 +1 @@
-version: 0.6.0
\ No newline at end of file
+version: 1.0.0
\ No newline at end of file
