Unreal Cloud DDC - Fixed issue were the ondemand replication could get stuck in a loop replicating from itself.

diff --git a/Engine/Source/Programs/UnrealCloudDDC/Jupiter/Implementation/Blob/BlobService.cs b/Engine/Source/Programs/UnrealCloudDDC/Jupiter/Implementation/Blob/BlobService.cs
--- a/Engine/Source/Programs/UnrealCloudDDC/Jupiter/Implementation/Blob/BlobService.cs
+++ b/Engine/Source/Programs/UnrealCloudDDC/Jupiter/Implementation/Blob/BlobService.cs
@@ -580,6 +580,8 @@ public class BlobService : IBlobService
 			throw new BlobReplicationException(ns, blob, "Blob not found in any region");
 		}
 
+		regions.RemoveAll(s => string.Equals(s, _currentSite, StringComparison.OrdinalIgnoreCase));
+
 		_logger.LogInformation("On-demand replicating blob {Blob} in Namespace {Namespace}", blob, ns);
 		List<(int, string)> possiblePeers = new List<(int, string)>(_peerStatusService.GetPeersByLatency(regions));
 
