Unreal Cloud DDC - Added allowOndemandReplication query parameter to blob fetching, only for internal use to prevent infininite recurrsion between regions attempting to replicate data they do not have.

diff --git a/Engine/Source/Programs/UnrealCloudDDC/Jupiter/Controllers/BlobController.cs b/Engine/Source/Programs/UnrealCloudDDC/Jupiter/Controllers/BlobController.cs
--- a/Engine/Source/Programs/UnrealCloudDDC/Jupiter/Controllers/BlobController.cs
+++ b/Engine/Source/Programs/UnrealCloudDDC/Jupiter/Controllers/BlobController.cs
@@ -49,7 +49,8 @@ namespace Jupiter.Controllers
 		public async Task<IActionResult> GetAsync(
 			[Required] NamespaceId ns,
 			[Required] BlobId id,
-			[FromQuery] List<string>? storageLayers = null)
+			[FromQuery] List<string>? storageLayers = null,
+			[FromQuery] bool allowOndemandReplication = true)
 		{
 			ActionResult? result = await _requestHelper.HasAccessToNamespaceAsync(User, Request, ns, new [] { JupiterAclAction.ReadObject });
 			if (result != null)
@@ -59,7 +60,7 @@ namespace Jupiter.Controllers
 
 			try
 			{
-				BlobContents blobContents = await GetImplAsync(ns, id, storageLayers, supportsRedirectUri: true);
+				BlobContents blobContents = await GetImplAsync(ns, id, storageLayers, supportsRedirectUri: true, allowOndemandReplication: allowOndemandReplication);
 
 				if (blobContents.RedirectUri != null)
 				{
@@ -155,15 +156,15 @@ namespace Jupiter.Controllers
 			return Ok(new HeadMultipleResponse { Needs = missingBlobs.ToArray()});
 		}
 
-		private async Task<BlobContents> GetImplAsync(NamespaceId ns, BlobId blob, List<string>? storageLayers = null, bool supportsRedirectUri = false)
+		private async Task<BlobContents> GetImplAsync(NamespaceId ns, BlobId blob, List<string>? storageLayers = null, bool supportsRedirectUri = false, bool allowOndemandReplication = true)
 		{
 			try
 			{
-				return await _storage.GetObjectAsync(ns, blob, storageLayers, supportsRedirectUri);
+				return await _storage.GetObjectAsync(ns, blob, storageLayers, supportsRedirectUri, allowOndemandReplication);
 			}
 			catch (BlobNotFoundException)
 			{
-				if (!_storage.ShouldFetchBlobOnDemand(ns))
+				if (!_storage.ShouldFetchBlobOnDemand(ns) || !allowOndemandReplication)
 				{
 					throw;
 				}
diff --git a/Engine/Source/Programs/UnrealCloudDDC/Jupiter/Implementation/Blob/BlobService.cs b/Engine/Source/Programs/UnrealCloudDDC/Jupiter/Implementation/Blob/BlobService.cs
--- a/Engine/Source/Programs/UnrealCloudDDC/Jupiter/Implementation/Blob/BlobService.cs
+++ b/Engine/Source/Programs/UnrealCloudDDC/Jupiter/Implementation/Blob/BlobService.cs
@@ -238,7 +238,7 @@ public class BlobService : IBlobService
 		return identifier;
 	}
 
-	public async Task<BlobContents> GetObjectAsync(NamespaceId ns, BlobId blob, List<string>? storageLayers = null, bool supportsRedirectUri = false)
+	public async Task<BlobContents> GetObjectAsync(NamespaceId ns, BlobId blob, List<string>? storageLayers = null, bool supportsRedirectUri = false, bool allowOndemandReplication = true)
 	{
 		try
 		{
@@ -296,7 +296,7 @@ public class BlobService : IBlobService
 				}
 			}
 
-			if (ShouldFetchBlobOnDemand(ns))
+			if (ShouldFetchBlobOnDemand(ns) && allowOndemandReplication)
 			{
 				try
 				{
@@ -309,7 +309,7 @@ public class BlobService : IBlobService
 			}
 
 			// if the primary namespace failed check to see if we should use a fallback policy which has replication enabled
-			if (policy.FallbackNamespace != null && ShouldFetchBlobOnDemand(policy.FallbackNamespace.Value))
+			if (policy.FallbackNamespace != null && ShouldFetchBlobOnDemand(policy.FallbackNamespace.Value) && allowOndemandReplication)
 			{
 				try
 				{
@@ -580,8 +580,6 @@ public class BlobService : IBlobService
 			throw new BlobReplicationException(ns, blob, "Blob not found in any region");
 		}
 
-		regions.RemoveAll(s => string.Equals(s, _currentSite, StringComparison.OrdinalIgnoreCase));
-
 		_logger.LogInformation("On-demand replicating blob {Blob} in Namespace {Namespace}", blob, ns);
 		List<(int, string)> possiblePeers = new List<(int, string)>(_peerStatusService.GetPeersByLatency(regions));
 
@@ -603,7 +601,7 @@ public class BlobService : IBlobService
 			{
 				url += "/";
 			}
-			using HttpRequestMessage blobRequest = await BuildHttpRequestAsync(HttpMethod.Get, new Uri($"{url}api/v1/blobs/{ns}/{blob}"));
+			using HttpRequestMessage blobRequest = await BuildHttpRequestAsync(HttpMethod.Get, new Uri($"{url}api/v1/blobs/{ns}/{blob}?allowOndemandReplication=false"));
 			HttpResponseMessage blobResponse = await httpClient.SendAsync(blobRequest, HttpCompletionOption.ResponseHeadersRead);
 
 			if (blobResponse.StatusCode == HttpStatusCode.NotFound)
diff --git a/Engine/Source/Programs/UnrealCloudDDC/Jupiter/Implementation/Blob/IBlobService.cs b/Engine/Source/Programs/UnrealCloudDDC/Jupiter/Implementation/Blob/IBlobService.cs
--- a/Engine/Source/Programs/UnrealCloudDDC/Jupiter/Implementation/Blob/IBlobService.cs
+++ b/Engine/Source/Programs/UnrealCloudDDC/Jupiter/Implementation/Blob/IBlobService.cs
@@ -21,7 +21,7 @@ public interface IBlobService
 	Task<BlobId> PutObjectAsync(NamespaceId ns, byte[] payload, BlobId identifier);
 	Task<Uri?> MaybePutObjectWithRedirectAsync(NamespaceId ns, BlobId identifier);
 
-	Task<BlobContents> GetObjectAsync(NamespaceId ns, BlobId blob, List<string>? storageLayers = null, bool supportsRedirectUri = false);
+	Task<BlobContents> GetObjectAsync(NamespaceId ns, BlobId blob, List<string>? storageLayers = null, bool supportsRedirectUri = false, bool allowOndemandReplication = true);
 	
 	Task<Uri?> GetObjectWithRedirectAsync(NamespaceId ns, BlobId blobIdentifier, List<string>? storageLayers = null);
 
