Unreal Cloud DDC - Added option to not put the inlined blob into the blob store, will reduce the total number of blobs in the blob store and these were never read as we used the inlined blob instead.

diff --git a/Engine/Source/Programs/UnrealCloudDDC/Helm/UnrealCloudDDC/templates/configmap.yaml b/Engine/Source/Programs/UnrealCloudDDC/Helm/UnrealCloudDDC/templates/configmap.yaml
--- a/Engine/Source/Programs/UnrealCloudDDC/Helm/UnrealCloudDDC/templates/configmap.yaml
+++ b/Engine/Source/Programs/UnrealCloudDDC/Helm/UnrealCloudDDC/templates/configmap.yaml
@@ -41,6 +41,10 @@ UnrealCloudDDC:
   EnableBucketStatsTracking: {{ $config.EnableBucketStatsTracking }}
 {{- end }}
 
+{{- if $config.EnablePutRefBodyIntoBlobStore }}
+  EnablePutRefBodyIntoBlobStore: {{ $config.EnablePutRefBodyIntoBlobStore }}
+{{- end }}
+
 {{ if $config.StorageImplementations }}
   StorageImplementations:
   {{- range $config.StorageImplementations -}}
diff --git a/Engine/Source/Programs/UnrealCloudDDC/Jupiter/Implementation/RefService.cs b/Engine/Source/Programs/UnrealCloudDDC/Jupiter/Implementation/RefService.cs
--- a/Engine/Source/Programs/UnrealCloudDDC/Jupiter/Implementation/RefService.cs
+++ b/Engine/Source/Programs/UnrealCloudDDC/Jupiter/Implementation/RefService.cs
@@ -135,9 +135,17 @@ namespace Jupiter.Implementation
 
 			Task objectStorePut = _referencesStore.PutAsync(ns, bucket, key, blobHash, payload.GetView().ToArray(), isFinalized);
 
-			Task<BlobId> blobStorePut = _blobService.PutObjectAsync(ns, payload.GetView().ToArray(), blobHash);
-			
-			await Task.WhenAll(objectStorePut, blobStorePut);
+			Task<BlobId>? blobStorePut = null;
+			if (_cloudDDCSettings.CurrentValue.EnablePutRefBodyIntoBlobStore)
+			{
+				blobStorePut = _blobService.PutObjectAsync(ns, payload.GetView().ToArray(), blobHash);
+			}
+
+			await objectStorePut;
+			if (blobStorePut != null)
+			{
+				await blobStorePut;
+			}
 
 			return await DoFinalizeAsync(ns, bucket, key, blobHash, payload);
 		}
diff --git a/Engine/Source/Programs/UnrealCloudDDC/Jupiter/Settings.cs b/Engine/Source/Programs/UnrealCloudDDC/Jupiter/Settings.cs
--- a/Engine/Source/Programs/UnrealCloudDDC/Jupiter/Settings.cs
+++ b/Engine/Source/Programs/UnrealCloudDDC/Jupiter/Settings.cs
@@ -115,6 +115,7 @@ namespace Jupiter
 		public bool EnableOnDemandReplication { get; set; } = true;
 
 		public bool EnableBucketStatsTracking { get; set; } = true;
+		public bool EnablePutRefBodyIntoBlobStore { get; set; } = true;
 	}
 
 	public class MongoSettings
