Unreal Cloud DDC - Fixes to helm charts to allow setting boolean values to false - resolves issue trying to use GCE s3 emulation.

diff --git a/Engine/Source/Programs/UnrealCloudDDC/Helm/UnrealCloudDDC/templates/configmap.yaml b/Engine/Source/Programs/UnrealCloudDDC/Helm/UnrealCloudDDC/templates/configmap.yaml
--- a/Engine/Source/Programs/UnrealCloudDDC/Helm/UnrealCloudDDC/templates/configmap.yaml
+++ b/Engine/Source/Programs/UnrealCloudDDC/Helm/UnrealCloudDDC/templates/configmap.yaml
@@ -29,19 +29,19 @@ UnrealCloudDDC:
   ContentIdStoreImplementation :  {{ default "Scylla" $config.ContentIdStoreImplementation }}
   BlobIndexImplementation :  {{ default "Scylla" $config.BlobIndexImplementation }}
 
-{{- if $config.EnableOnDemandReplication }}
+{{- if hasKey $config "EnableOnDemandReplication" }}
   EnableOnDemandReplication :  {{ $config.EnableOnDemandReplication }}
 {{- end }}
 
-{{- if $config.EnableLastAccessTracking }}
+{{- if hasKey $config "EnableLastAccessTracking" }}
   EnableLastAccessTracking: {{ $config.EnableLastAccessTracking }}
 {{- end }}
 
-{{- if $config.EnableBucketStatsTracking }}
+{{- if hasKey $config "EnableBucketStatsTracking" }}
   EnableBucketStatsTracking: {{ $config.EnableBucketStatsTracking }}
 {{- end }}
 
-{{- if $config.EnablePutRefBodyIntoBlobStore }}
+{{- if hasKey $config "EnablePutRefBodyIntoBlobStore" }}
   EnablePutRefBodyIntoBlobStore: {{ $config.EnablePutRefBodyIntoBlobStore }}
 {{- end }}
 
@@ -175,7 +175,7 @@ Azure:
 {{- $config := .config -}}
 {{- $s3 := default dict $config.S3 -}}
 S3:
-  ConnectionString: "AWS"
+  ConnectionString: {{ default "AWS" $s3.ConnectionString }} 
   BucketName: {{ required "The name of the S3 bucket to use has to be specified (.config.S3.BucketName)" $s3.BucketName }}
   # We assume that you are creating your buckets yourself for production deployments like this 
   CreateBucketIfMissing: {{ default  false $s3.CreateBucketIfMissing }}
@@ -183,7 +183,7 @@ S3:
 
   UseBlobIndexForExistsCheck: {{ default false $s3.UseBlobIndexForExistsCheck }}
 
-  {{- if $s3.UseArnRegion }}
+  {{- if hasKey $s3 "UseArnRegion" }}
   UseArnRegion: {{ $s3.UseArnRegion }}
   {{- end }}
 
@@ -191,11 +191,11 @@ S3:
   StoragePoolBucketOverride: {{ $s3.StoragePoolBucketOverride | toYaml | nindent 4  }}
   {{- end }}
 
-  {{- if $s3.UseMultiPartUpload }}
+  {{- if hasKey $s3 "UseMultiPartUpload" }}
   UseMultiPartUpload: {{ $s3.UseMultiPartUpload }}
   {{- end }}
 
-  {{- if $s3.UseChunkEncoding }}
+  {{- if hasKey $s3 "UseChunkEncoding" }}
   UseChunkEncoding: {{ $s3.UseChunkEncoding }}
   {{- end }}
   
