Unreal Cloud DDC - Added a cloud bootstrapper script for AWS the idea is that users with limited cloud experience can use this to get a cloud ddc setup created for them.

diff --git a/Engine/Source/Programs/UnrealCloudDDC/Bootstrapping/AWS/cloud-ddc.yaml b/Engine/Source/Programs/UnrealCloudDDC/Bootstrapping/AWS/cloud-ddc.yaml
new file mode 100644
--- /dev/null
+++ b/Engine/Source/Programs/UnrealCloudDDC/Bootstrapping/AWS/cloud-ddc.yaml
@@ -0,0 +1,175 @@
+######
+# AWS configuration
+######
+
+serviceAccount:
+  name: svc-unreal-cloud-ddc
+  create: false
+
+global:
+  cloudProvider: "AWS"
+  awsRegion: "%REGION_REPLACE%"
+  awsRole: AssumeRoleWebIdentity # service accounts are IAM mapped
+  
+  siteName: "%REGION_REPLACE%"
+
+######
+# Set a disabled bearer handler, update this to your SSO
+######
+  auth:
+    defaultScheme: ServiceAccount
+    #defaultScheme: Bearer
+    schemes:
+      ServiceAccount: 
+        implementation: ServiceAccount
+      #Bearer: # example integration of a standard SSO
+      #  implementation: JWTBearer
+      #  JwtAudience: api://audience
+      #  JwtAuthority: https://url-to-your-server
+
+      #Bearer: # example integration integrating with Okta
+      #  implementation: Okta
+      #  OktaDomain: https://url-to-your-server
+      #  OktaAuthorizationServerId: default
+      #  JwtAudience: api://audience
+      #  JwtAuthority: http://url-to-your-server/oauth2/auth_server_Id
+    acls:
+    - claims: 
+      - groups=cloud-ddc-admin
+      actions:
+      - ReadObject
+      - WriteObject
+      - DeleteObject
+      - DeleteBucket
+      - DeleteNamespace
+      - AdminAction
+
+######
+# Map the claims coming from your SSO or the service account to actual actions in Unreal CLoud DDC - see `Namespace access` in the cloud ddc readme for information on what actions can be mapped
+######
+
+  namespaces:
+    Policies:
+      ddc:
+        acls:
+        - actions: 
+          - ReadObject
+          - WriteObject
+          claims: 
+          - groups=app-cloud-ddc-project
+######
+# Kubernetes configuration
+######
+
+nodeSelector:
+  cloud-ddc.epicgames.com/node-type: nvme
+
+tolerations:
+  - key: role
+    operator: Equal
+    value: unreal-cloud-ddc
+    effect: NoSchedule
+
+config:
+
+######
+# Scylla configuration
+######
+  Scylla: &scylla-connection-info
+    ConnectionString: Contact Points=scylla-client.scylla.svc.cluster.local;Default Keyspace=jupiter;
+
+    KeyspaceReplicationStrategy:
+      class : "NetworkTopologyStrategy"
+      %REGION_REPLACE%: 2
+
+    LocalDatacenterName: "%REGION_REPLACE%"
+    LocalKeyspaceSuffix: "%REGION_REPLACE_CLEAN%"
+
+######
+# S3 configuration
+######
+  S3: &s3-connection-info
+    BucketName: %S3BUCKET_NAME%
+
+######
+# Auth configuration for the service token - this can be removed once you have updated to your SSO
+######
+  ServiceAccounts:
+    Accounts:
+    - Token: %SERVICE_ACCOUNT_TOKEN%
+      Claims:
+      - groups:app-cloud-ddc-project
+      - groups:cloud-ddc-admin
+
+######
+# Worker configuration - setup to run GC right now - update the replication configuration if you want to enable replication between regions
+######
+
+worker:
+  enabled: true
+
+  imagePullSecrets: 
+    - name: regcred
+
+  nodeSelector:
+    cloud-ddc.epicgames.com/node-type: general
+
+  tolerations:
+    - key: role
+      operator: Equal
+      value: unreal-cloud-ddc
+      effect: NoSchedule
+  
+  config:
+    Scylla: *scylla-connection-info
+    S3: *s3-connection-info
+    
+    GC:
+      # should only be true for a single region
+      CleanOldRefRecords: true
+      CleanOldBlobs: true
+
+    Replication:
+      Enabled: false
+
+######
+# Mount nvme drive
+######
+persistence:
+  enabled: false
+  volume:     
+    hostPath:
+      path: /data-nvme # we set up this to run on our attached nvme disk
+      type: Directory
+
+######
+# Enable nginx proxy
+######
+nginx:
+  enabled: true
+  useDomainSockets: true
+
+######
+# Create a NLB
+######
+service:
+  annotations:
+    service.beta.kubernetes.io/aws-load-balancer-name: "unreal-cloud-ddc"
+    service.beta.kubernetes.io/aws-load-balancer-type: "external"
+    service.beta.kubernetes.io/aws-load-balancer-nlb-target-type: "ip"
+    service.beta.kubernetes.io/aws-load-balancer-scheme: "internet-facing"
+    service.beta.kubernetes.io/aws-load-balancer-healthcheck-protocol: "HTTP"
+    service.beta.kubernetes.io/aws-load-balancer-healthcheck-port: "80"
+    service.beta.kubernetes.io/aws-load-balancer-healthcheck-path: "/health/live"
+    service.beta.kubernetes.io/aws-load-balancer-healthcheck-healthy-threshold: "2"
+    service.beta.kubernetes.io/aws-load-balancer-healthcheck-unhealthy-threshold: "2"
+    service.beta.kubernetes.io/aws-load-balancer-additional-resource-tags: application=unreal-cloud-ddc
+
+  type: LoadBalancer
+  portName: http
+  port: 80
+  targetPort: nginx-http
+  extraPort:
+  - name: http2
+    port: 8080
+    targetPort: nginx-h2
\ No newline at end of file
diff --git a/Engine/Source/Programs/UnrealCloudDDC/Bootstrapping/AWS/cmdlets.ps1 b/Engine/Source/Programs/UnrealCloudDDC/Bootstrapping/AWS/cmdlets.ps1
new file mode 100644
--- /dev/null
+++ b/Engine/Source/Programs/UnrealCloudDDC/Bootstrapping/AWS/cmdlets.ps1
@@ -0,0 +1,22 @@
+
+function Upload-AWSBootstrapperToS3Dev {
+    $PSNativeCommandUseErrorActionPreference = $true; 
+    $ErrorActionPreference = 'Stop';
+
+    Write-Host "Compressing configuration files"
+    Get-ChildItem -Path $PSScriptRoot -Exclude *.ps1,*.zip |  Compress-Archive -CompressionLevel Fastest -DestinationPath setup.zip -Force
+
+    Write-Host "Uploading to S3"
+    aws s3 cp setup.zip s3://devtools-misc/dev-cloud-ddc-bootstrapper/aws/setup.zip --acl public-read
+}
+
+function Upload-AWSBootstrapperToS3Prod {
+    $PSNativeCommandUseErrorActionPreference = $true; 
+    $ErrorActionPreference = 'Stop';
+
+    Write-Host "Compressing configuration files"
+    Get-ChildItem -Path $PSScriptRoot -Exclude *.ps1,*.zip |  Compress-Archive -CompressionLevel Fastest -DestinationPath setup.zip -Force
+
+    Write-Host "Uploading to S3"
+    aws s3 cp setup.zip s3://devtools-misc/cloud-ddc-bootstrapper/aws/setup.zip --acl public-read
+}
diff --git a/Engine/Source/Programs/UnrealCloudDDC/Bootstrapping/AWS/delete.sh b/Engine/Source/Programs/UnrealCloudDDC/Bootstrapping/AWS/delete.sh
new file mode 100644
--- /dev/null
+++ b/Engine/Source/Programs/UnrealCloudDDC/Bootstrapping/AWS/delete.sh
@@ -0,0 +1,39 @@
+#! /bin/bash
+set -euo pipefail
+
+REGION_REPLACE=$AWS_REGION
+AZ_REPLACE_0=$(aws ec2 describe-availability-zones | jq -r .AvailabilityZones[0].ZoneName)
+AZ_REPLACE_1=$(aws ec2 describe-availability-zones | jq -r .AvailabilityZones[1].ZoneName)
+
+echo "Removing cloud ddc from region: \"$REGION_REPLACE\" using Availability zones: \"$AZ_REPLACE_0\" and \"$AZ_REPLACE_1\""
+
+sed -i "s/\%REGION_REPLACE\%/$REGION_REPLACE/g" eksctl.yaml
+sed -i "s/\%REGION_REPLACE\%/$REGION_REPLACE/g" cloud-ddc.yaml
+
+sed -i "s/\%REGION_REPLACE\%/$REGION_REPLACE/g" scylla.cluster.yaml
+sed -i "s/\%AZ_REPLACE_0\%/$AZ_REPLACE_0/g" scylla.cluster.yaml
+sed -i "s/\%AZ_REPLACE_1\%/$AZ_REPLACE_1/g" scylla.cluster.yaml
+
+echo "Removing Scylla monitoring..."
+kubectl -n scylla delete ScyllaDBMonitoring/scylla-monitoring
+
+echo "Removing Helm charts..."
+helm uninstall unreal-cloud-ddc -n unreal-cloud-ddc
+
+helm uninstall scylla -n scylla
+helm uninstall scylla-manager -n scylla-manager
+helm uninstall scylla-operator -n scylla-operator
+
+helm uninstall prometheus-operator --namespace prometheus-operator
+
+echo "Removing Scylla csi driver..."
+kubectl -n local-csi-driver delete -f ./scylla-provisioner/local-csi-driver
+
+kubectl delete pods -n scylla --all --force
+echo "Warning! Not deleting S3 bucket, you will need to manually delete this ..."
+
+echo "Deleting cluster..."
+eksctl delete addon --name aws-ebs-csi-driver --cluster unreal-cloud-ddc
+eksctl delete addon --name coredns --cluster unreal-cloud-ddc
+
+eksctl delete cluster -f ./eksctl.yaml --parallel 4 --disable-nodegroup-eviction
\ No newline at end of file
diff --git a/Engine/Source/Programs/UnrealCloudDDC/Bootstrapping/AWS/eksctl.yaml b/Engine/Source/Programs/UnrealCloudDDC/Bootstrapping/AWS/eksctl.yaml
new file mode 100644
--- /dev/null
+++ b/Engine/Source/Programs/UnrealCloudDDC/Bootstrapping/AWS/eksctl.yaml
@@ -0,0 +1,79 @@
+apiVersion: eksctl.io/v1alpha5
+kind: ClusterConfig
+
+metadata:
+  name: unreal-cloud-ddc
+  region: %REGION_REPLACE%
+  tags:
+    created-by: CloudDDCBootstrapper
+
+vpc:
+  cidr: 192.168.0.0/19
+
+availabilityZones:
+- %AZ_REPLACE_0%
+- %AZ_REPLACE_1%
+
+addons:
+- name: vpc-cni # no version is specified so it deploys the default version
+  attachPolicyARNs:
+    - arn:aws:iam::aws:policy/AmazonEKS_CNI_Policy
+- name: coredns
+  version: latest # auto discovers the latest available
+- name: kube-proxy
+  version: latest
+- name: aws-ebs-csi-driver
+  version: latest
+
+managedNodeGroups:
+  - name: nvme
+    instanceType: i3en.xlarge
+    desiredCapacity: 2
+    labels:
+      cloud-ddc.epicgames.com/node-type: nvme
+    taints:
+    - key: role
+      value: unreal-cloud-ddc
+      effect: NoSchedule
+    preBootstrapCommands:
+    - sudo mkfs.ext4 -E nodiscard /dev/nvme1n1
+    - sudo mkdir /data-nvme
+    - sudo mount /dev/nvme1n1 /data-nvme
+  - name: general 
+    instanceType: c5.xlarge
+    desiredCapacity: 1
+    labels:
+      cloud-ddc.epicgames.com/node-type: general
+
+nodeGroups:
+  - name: scylla-pool
+    instanceType: i4i.xlarge
+    desiredCapacity: 2
+    labels:
+      scylla.scylladb.com/node-type: scylla
+    taints:
+    - key: role
+      value: scylla-clusters
+      effect: NoSchedule
+    kubeletExtraConfig:
+      cpuManagerPolicy: static
+
+iam:
+  withOIDC: true
+  serviceAccounts:
+  - metadata:
+      name: svc-unreal-cloud-ddc
+      namespace: unreal-cloud-ddc
+      labels: {aws-usage: "application"}
+    attachPolicyARNs:
+    - "arn:aws:iam::aws:policy/AmazonS3FullAccess"
+  - metadata:
+      name: external-dns
+      namespace: kube-system
+    wellKnownPolicies:
+      externalDNS: true
+  - metadata:
+      name: aws-load-balancer-controller
+      namespace: kube-system
+    wellKnownPolicies:
+      awsLoadBalancerController: true
diff --git a/Engine/Source/Programs/UnrealCloudDDC/Bootstrapping/AWS/nodeconfig-alpha.yaml b/Engine/Source/Programs/UnrealCloudDDC/Bootstrapping/AWS/nodeconfig-alpha.yaml
new file mode 100644
--- /dev/null
+++ b/Engine/Source/Programs/UnrealCloudDDC/Bootstrapping/AWS/nodeconfig-alpha.yaml
@@ -0,0 +1,29 @@
+apiVersion: scylla.scylladb.com/v1alpha1
+kind: NodeConfig
+metadata:
+  name: cluster
+spec:
+  localDiskSetup:
+    filesystems:
+    - device: /dev/md/nvmes
+      type: xfs
+    mounts:
+    - device: /dev/md/nvmes
+      mountPoint: /mnt/persistent-volumes
+      unsupportedOptions:
+      - prjquota
+    raids:
+    - name: nvmes
+      type: RAID0
+      RAID0:
+        devices:
+          modelRegex: Amazon EC2 NVMe Instance Storage
+          nameRegex: ^/dev/nvme\d+n\d+$
+  placement:
+    nodeSelector:
+      scylla.scylladb.com/node-type: scylla
+    tolerations:
+    - effect: NoSchedule
+      key: role
+      operator: Equal
+      value: scylla-clusters
\ No newline at end of file
diff --git a/Engine/Source/Programs/UnrealCloudDDC/Bootstrapping/AWS/readme.md b/Engine/Source/Programs/UnrealCloudDDC/Bootstrapping/AWS/readme.md
new file mode 100644
--- /dev/null
+++ b/Engine/Source/Programs/UnrealCloudDDC/Bootstrapping/AWS/readme.md
@@ -0,0 +1,130 @@
+This helps you setup a single region deployment of Unreal Cloud DDC on AWS if you have no or very limited previous experience setting up services on AWS.
+
+If you are more familar with how AWS works we recommend using a IaC (Infrastructure as code) system to define these resources instead.
+You can use this as a example of how to setup it up but its also possible to set this up without EKS and just using EC2 if you prefer.
+
+## Prereqs
+* An AWS account
+
+## Setup
+
+To install go to the AWS Management console and start AWS Cloudshell by hitting the >_ button on the top right side of the page (alternativley use the button on the bottom left labelled `CloudShell`).
+Make sure you have selected the correct region you wish to create this in before starting AWS CloudShell.
+
+Once this shell has started up run this command (note that this will take an hour or so to run, your cloud shell can timeout during this time but the command will continue to run):
+```
+nohup bash -c 'curl --silent -LO https://devtools-misc.s3.amazonaws.com/cloud-ddc-bootstrapper/aws/setup.zip && unzip -oq setup.zip && chmod +x setup.sh && ./setup.sh'
+```
+
+This command will do the following:
+* Create a AWS EKS Cluster
+* Setup the EKS cluster to run a instance of Scylla Open Source using Scylla Operator
+* Install Cloud DDC onto the same cluster
+* Create a S3 bucket for storage of blobs
+* Setup simple service account auth (one shared key that can be used by everyone).
+
+### Cost examples
+This generates a bit of hardware that has flat costs so you should expect a run of this to end up with a cost about $1600 / month broken down like this:
+* i3en instances - $700 (these are used to host the cloud ddc api)
+* i3 instances - $120 (this hosts the scylla monitoring and mangement services)
+* i4i instances - $500 (this hosts the actual scylla database)
+* c5 instances - $130 (used to host the cloud ddc worker)
+
+On top of this you will see cost for the load balancer and s3 bucket based on usage and some other minor AWS fees.
+These cost are list price and without prepurchashing hardware.
+
+### Making the cluster production ready
+After completing the setup here follows some suggestions for preparing this cluster for production use which we strongly recommend you follow.
+
+#### Setup nicer domain name
+This requires you to setup a Route53 hosted zone which your domain can be in, AWS provides a guide for that [here|https://us-east-1.console.aws.amazon.com/route53/domains/home?region=us-east-1#/DomainSearch] .
+Once this hosted zone is setup you need to add a annotation to your svc to indicate the domain name you want it to use, this can be achived by running this command:
+`kubectl annotate svc unreal-cloud-ddc -n unreal-cloud-ddc "external-dns.alpha.kubernetes.io/hostname=<fully-qualifed-url-you-want-to-use>"`
+After a few minutes a CName record should have been created in Route53 that points to your load balancer.
+
+Alternativley update the `cloud-ddc.yaml` file with this annotation and run
+`helm upgrade unreal-cloud-ddc oci://ghcr.io/epicgames/unreal-cloud-ddc --namespace unreal-cloud-ddc --values cloud-ddc.yaml`
+
+#### Setup https
+To use https you will need to allocate a certificate in AWS Cert Manager, you will usually need a domain name associated with it so its recommended to follow the `Setup nicer domain name` steps first.
+Once you have a certificate allocated in Cert Manager you should take the ARN of this certificate and run this this command:
+`kubectl annotate svc unreal-cloud-ddc -n unreal-cloud-ddc "service.beta.kubernetes.io/aws-load-balancer-ssl-ports=443,8080" "service.beta.kubernetes.io/aws-load-balancer-ssl-cert=<add-arn-here>"`
+Alternativley update the `cloud-ddc.yaml` file with this annotation and run
+`helm upgrade unreal-cloud-ddc oci://ghcr.io/epicgames/unreal-cloud-ddc --namespace unreal-cloud-ddc --values cloud-ddc.yaml`
+
+
+After you have done this you should update load balancer to correctly report which http versions are supported which you can do using these commands:
+```
+aws modify-listener --listener-arn <arn-to-load-balancer-listener-port-80> --alpn-policy HTTP1Only
+aws modify-listener --listener-arn <arn-to-load-balancer-listener-port-8080> --alpn-policy HTTP2Only
+```
+
+#### Enable SSO Auth
+We strongly recommended to change auth mechanism to use a SSO method that supports `auth code` logins, which most SSOs will. To configure this on the editor side see `IdentityProvider setup` in the UnrealCloudDDC readme .
+To change your deployment you will need to know the jwtAudience and jwtAuthority your SSO will usually tell you what these are. If you are using Okta we provide a special setup that integrates easier.
+To update these settings edit the `cloud-ddc.yaml` file and change the `auth` section. For the standard SSO setup you should change it to this:
+```
+  auth:
+    defaultScheme: Bearer
+    schemes:
+      ServiceAccount: 
+        implementation: ServiceAccount
+      Bearer: # example integration of a standard SSO
+        implementation: JWTBearer
+        JwtAudience: api://audience
+        JwtAuthority: https://url-to-your-server
+```
+Note that we kept the ServiceAccount setup we use by default around here, this means that we accept both SSO logins and these service accounts. This can be useful while transitioning but its recommended to remove the ServiceAccount scheme once you have finished transitioning.
+
+The okta setup is similar
+```
+  auth:
+    defaultScheme: Bearer
+    schemes:
+      ServiceAccount: 
+        implementation: ServiceAccount
+      Bearer: # example integration integrating with Okta
+        implementation: Okta
+        OktaDomain: https://url-to-your-server
+        OktaAuthorizationServerId: default
+        JwtAudience: api://audience
+        JwtAuthority: http://url-to-your-server/oauth2/auth_server_Id
+```
+
+Once you have updated the `cloud-ddc.yaml` you can run this command to update the kubernetes setup with the new configuration
+`helm upgrade unreal-cloud-ddc oci://ghcr.io/epicgames/unreal-cloud-ddc --namespace unreal-cloud-ddc --values cloud-ddc.yaml`
+
+## Other useful topics
+Next follows some other useful topics
+
+### Remoting into the machines
+We do not setup ssh into any machines, but you can use AWS SSM to access the machines, for that you will need to know the EC2 instance id. To remote into that machine simply use this command:
+`aws ssm start-session --target <ec2-instance-id like i-0000000000abc>`
+You can also remote into the kubernetes pods themselves, this can be useful to inspect the local state of the filesystem cache.
+First you need to find the pod name you want to remote into `kubectl get pods -n unreal-cloud-ddc` will show you all the pods and you can then use
+`kubectl exec -it <pod-name> -c unreal-cloud-ddc -n unreal-cloud-ddc -- /bin/bash` to remote into the pod.
+
+### Going multiregion
+Scaling out to multiple regions is quite complicated, we recommend you make sure to consult with someone that understands the best way to achive that within your AWS setup and in a way that is compatible with your security considerations.
+
+Note that you will need to aquire a Scylla Enterprise license before being able to scale to a multiregion deployment using Scylla Operator. Alternativley you can build your own manual cluster on EC2 using just Scylla Open Source.
+
+Steps to take
+1. The configurations generated by this script can be used as a foundation but you will need to run it manually. Start by updating these files:
+* Update `eksctl.yaml` to use another of the `192.168` CIDRs.
+* Update `scylla-cluster.yaml` - This may need to update the region and AZ specifications. You will also need to expose the cluster between your old and new VPC - described in this documentation: https://operator.docs.scylladb.com/stable/exposing.html#multi-vpc
+* Update `cloud-ddc.yaml` - Adjust any region specific settings in here.
+2. Update this new VPC to be peered with your previous VPC and setup routing between the regions.
+3. Create a new S3 bucket for the new region.
+3. Remote into one of your scylla nodes, and run `cqlsh` to update the Scylla replication factor.
+4. Manually trigger a repair of scylla using the Scylla Manager.
+5. Update the configuration of Unreal Cloud DDC to setup the replication between each region (see `Blob replication setup` in readme for Unreal Cloud DDC).
+6. Setup a Route53 latency record that includes the public host names of both of your regions, update your UE ini files to use this latency record.
+
+
+## Deleting the cluster
+If you wish to remove the cluster you can run this script:
+./delete.sh
+
+If you ran this due to setup.sh failing it may not be able to delete everything in which case you can manually delete the EKS cluster by running this command from AWS CloudShell:
+`eksctl delete cluster -f ./eksctl.yaml --parallel 4 --disable-nodegroup-eviction`
\ No newline at end of file
diff --git a/Engine/Source/Programs/UnrealCloudDDC/Bootstrapping/AWS/scylla-monitoring.yaml b/Engine/Source/Programs/UnrealCloudDDC/Bootstrapping/AWS/scylla-monitoring.yaml
new file mode 100644
--- /dev/null
+++ b/Engine/Source/Programs/UnrealCloudDDC/Bootstrapping/AWS/scylla-monitoring.yaml
@@ -0,0 +1,27 @@
+apiVersion: scylla.scylladb.com/v1alpha1
+kind: ScyllaDBMonitoring
+metadata:
+  name: scylla-monitoring
+  namespace: scylla
+spec:
+  type: Platform
+  endpointsSelector:
+    matchLabels:
+      app.kubernetes.io/name: scylla
+      scylla-operator.scylladb.com/scylla-service-type: identity
+      scylla/cluster: scylla
+  components:
+    prometheus:
+      storage:
+        volumeClaimTemplate:
+          spec:
+            resources:
+              requests:
+                storage: 1Gi
+    grafana:
+      exposeOptions:
+        webInterface:
+          ingress:
+            ingressClassName: alb
+            dnsDomains:
+            - scylla-grafana.svc.cluster.local
\ No newline at end of file
diff --git a/Engine/Source/Programs/UnrealCloudDDC/Bootstrapping/AWS/scylla-provisioner/local-csi-driver/00_namespace.yaml b/Engine/Source/Programs/UnrealCloudDDC/Bootstrapping/AWS/scylla-provisioner/local-csi-driver/00_namespace.yaml
new file mode 100644
--- /dev/null
+++ b/Engine/Source/Programs/UnrealCloudDDC/Bootstrapping/AWS/scylla-provisioner/local-csi-driver/00_namespace.yaml
@@ -0,0 +1,4 @@
+apiVersion: v1
+kind: Namespace
+metadata:
+  name: local-csi-driver
diff --git a/Engine/Source/Programs/UnrealCloudDDC/Bootstrapping/AWS/scylla-provisioner/local-csi-driver/10_csidriver.yaml b/Engine/Source/Programs/UnrealCloudDDC/Bootstrapping/AWS/scylla-provisioner/local-csi-driver/10_csidriver.yaml
new file mode 100644
--- /dev/null
+++ b/Engine/Source/Programs/UnrealCloudDDC/Bootstrapping/AWS/scylla-provisioner/local-csi-driver/10_csidriver.yaml
@@ -0,0 +1,7 @@
+apiVersion: storage.k8s.io/v1
+kind: CSIDriver
+metadata:
+  name: local.csi.scylladb.com
+spec:
+  attachRequired: false
+  storageCapacity: true
diff --git a/Engine/Source/Programs/UnrealCloudDDC/Bootstrapping/AWS/scylla-provisioner/local-csi-driver/10_driver_serviceaccount.yaml b/Engine/Source/Programs/UnrealCloudDDC/Bootstrapping/AWS/scylla-provisioner/local-csi-driver/10_driver_serviceaccount.yaml
new file mode 100644
--- /dev/null
+++ b/Engine/Source/Programs/UnrealCloudDDC/Bootstrapping/AWS/scylla-provisioner/local-csi-driver/10_driver_serviceaccount.yaml
@@ -0,0 +1,5 @@
+apiVersion: v1
+kind: ServiceAccount
+metadata:
+  name: local-csi-driver
+  namespace: local-csi-driver
diff --git a/Engine/Source/Programs/UnrealCloudDDC/Bootstrapping/AWS/scylla-provisioner/local-csi-driver/10_provisioner_clusterrole.yaml b/Engine/Source/Programs/UnrealCloudDDC/Bootstrapping/AWS/scylla-provisioner/local-csi-driver/10_provisioner_clusterrole.yaml
new file mode 100644
--- /dev/null
+++ b/Engine/Source/Programs/UnrealCloudDDC/Bootstrapping/AWS/scylla-provisioner/local-csi-driver/10_provisioner_clusterrole.yaml
@@ -0,0 +1,100 @@
+kind: ClusterRole
+apiVersion: rbac.authorization.k8s.io/v1
+metadata:
+  name: scylladb:csi-external-provisioner
+rules:
+- apiGroups:
+  - ""
+  resources:
+  - "persistentvolumes"
+  verbs:
+  - "get"
+  - "list"
+  - "watch"
+  - "create"
+  - "delete"
+- apiGroups:
+  - ""
+  resources:
+  - "persistentvolumeclaims"
+  verbs:
+  - "get"
+  - "list"
+  - "watch"
+  - "update"
+- apiGroups:
+  - "storage.k8s.io"
+  resources:
+  - "storageclasses"
+  verbs:
+  - "get"
+  - "list"
+  - "watch"
+- apiGroups:
+  - ""
+  resources:
+  - "events"
+  verbs:
+  - "list"
+  - "watch"
+  - "create"
+  - "update"
+  - "patch"
+- apiGroups:
+  - "snapshot.storage.k8s.io"
+  resources:
+  - "volumesnapshots"
+  verbs:
+  - "get"
+  - "list"
+- apiGroups:
+  - "snapshot.storage.k8s.io"
+  resources:
+  - "volumesnapshotcontents"
+  verbs:
+  - "get"
+  - "list"
+- apiGroups:
+  - "storage.k8s.io"
+  resources:
+  - "csinodes"
+  verbs:
+  - "get"
+  - "list"
+  - "watch"
+- apiGroups:
+  - ""
+  resources:
+  - "nodes"
+  verbs:
+  - "get"
+  - "list"
+  - "watch"
+- apiGroups:
+  - "storage.k8s.io"
+  resources:
+  - "csistoragecapacities"
+  verbs:
+  - "get"
+  - "list"
+  - "watch"
+  - "create"
+  - "update"
+  - "patch"
+  - "delete"
+# The GET permissions below are needed for walking up the ownership chain
+# for CSIStorageCapacity. They are sufficient for deployment via
+# StatefulSet (only needs to get Pod) and Deployment (needs to get
+# Pod and then ReplicaSet to find the Deployment).
+- apiGroups:
+  - ""
+  resources:
+  - "pods"
+  verbs:
+  - "get"
+- apiGroups:
+  - "apps"
+  resources:
+  - "replicasets"
+  verbs:
+  - "get"
diff --git a/Engine/Source/Programs/UnrealCloudDDC/Bootstrapping/AWS/scylla-provisioner/local-csi-driver/20_provisioner_clusterrolebinding.yaml b/Engine/Source/Programs/UnrealCloudDDC/Bootstrapping/AWS/scylla-provisioner/local-csi-driver/20_provisioner_clusterrolebinding.yaml
new file mode 100644
--- /dev/null
+++ b/Engine/Source/Programs/UnrealCloudDDC/Bootstrapping/AWS/scylla-provisioner/local-csi-driver/20_provisioner_clusterrolebinding.yaml
@@ -0,0 +1,12 @@
+kind: ClusterRoleBinding
+apiVersion: rbac.authorization.k8s.io/v1
+metadata:
+  name: scylladb:csi-external-provisioner
+subjects:
+- kind: ServiceAccount
+  name: local-csi-driver
+  namespace: local-csi-driver
+roleRef:
+  kind: ClusterRole
+  name: scylladb:csi-external-provisioner
+  apiGroup: rbac.authorization.k8s.io
diff --git a/Engine/Source/Programs/UnrealCloudDDC/Bootstrapping/AWS/scylla-provisioner/local-csi-driver/50_daemonset.yaml b/Engine/Source/Programs/UnrealCloudDDC/Bootstrapping/AWS/scylla-provisioner/local-csi-driver/50_daemonset.yaml
new file mode 100644
--- /dev/null
+++ b/Engine/Source/Programs/UnrealCloudDDC/Bootstrapping/AWS/scylla-provisioner/local-csi-driver/50_daemonset.yaml
@@ -0,0 +1,128 @@
+kind: DaemonSet
+apiVersion: apps/v1
+metadata:
+  name: local-csi-driver
+  namespace: local-csi-driver
+  labels:
+    app.kubernetes.io/name: local-csi-driver
+spec:
+  selector:
+    matchLabels:
+      app.kubernetes.io/name: local-csi-driver
+  template:
+    metadata:
+      labels:
+        app.kubernetes.io/name: local-csi-driver
+    spec:
+      nodeSelector:
+        scylla.scylladb.com/node-type: scylla
+        kubernetes.io/os: linux
+      serviceAccountName: local-csi-driver
+      tolerations:
+      - effect: NoSchedule
+        key: role
+        operator: Equal
+        value: scylla-clusters
+      containers:
+      - name: local-csi-driver
+        securityContext:
+          privileged: true
+        image: docker.io/scylladb/k8s-local-volume-provisioner:0.1.0@sha256:d8c12e96d352077fd4b639c8e4fb050e49e0be11161d78d9339b69f73466f59a
+        imagePullPolicy: Always
+        args:
+        - --listen=/csi/csi.sock
+        - --node-name=$(NODE_NAME)
+        - --volumes-dir=/mnt/persistent-volumes
+        - --v=2
+        env:
+        - name: NODE_NAME
+          valueFrom:
+            fieldRef:
+              fieldPath: spec.nodeName
+        volumeMounts:
+        - name: kubelet-dir
+          mountPath: /var/lib/kubelet
+          mountPropagation: "Bidirectional"
+        - name: plugin-dir
+          mountPath: /csi
+        - name: volumes-dir
+          mountPath: /mnt/persistent-volumes
+        ports:
+        - name: healthz
+          containerPort: 9809
+          protocol: TCP
+        livenessProbe:
+          httpGet:
+            path: /healthz
+            port: healthz
+          initialDelaySeconds: 10
+          timeoutSeconds: 3
+          periodSeconds: 2
+          failureThreshold: 5
+      - name: csi-driver-registrar
+        image: registry.k8s.io/sig-storage/csi-node-driver-registrar@sha256:fdff3ee285341bc58033b6b2458a5d45fd90ec6922a8ba6ebdd49b0c41e2cd34
+        imagePullPolicy: Always
+        args:
+        - --csi-address=/csi/csi.sock
+        - --kubelet-registration-path=/var/lib/kubelet/plugins/local.csi.scylladb.com/csi.sock
+        volumeMounts:
+        - name: plugin-dir
+          mountPath: /csi
+        - name: registration-dir
+          mountPath: /registration
+      - name: liveness-probe
+        image: registry.k8s.io/sig-storage/livenessprobe@sha256:cacee2b5c36dd59d4c7e8469c05c9e4ef53ecb2df9025fa8c10cdaf61bce62f0
+        imagePullPolicy: Always
+        args:
+        - --csi-address=/csi/csi.sock
+        - --health-port=9809
+        - --v=2
+        volumeMounts:
+        - name: plugin-dir
+          mountPath: /csi
+      - name: csi-provisioner
+        image: registry.k8s.io/sig-storage/csi-provisioner@sha256:ee3b525d5b89db99da3b8eb521d9cd90cb6e9ef0fbb651e98bb37be78d36b5b8
+        imagePullPolicy: Always
+        args:
+        - --csi-address=/csi/csi.sock
+        - --v=2
+        - --node-deployment
+        - --feature-gates=Topology=true
+        - --immediate-topology=false
+        - --enable-capacity
+        - --capacity-ownerref-level=0
+        - --capacity-poll-interval=30s
+        - --default-fstype=xfs
+        env:
+        - name: NAMESPACE
+          valueFrom:
+            fieldRef:
+              fieldPath: metadata.namespace
+        - name: POD_NAME
+          valueFrom:
+            fieldRef:
+              fieldPath: metadata.name
+        - name: NODE_NAME
+          valueFrom:
+            fieldRef:
+              fieldPath: spec.nodeName
+        volumeMounts:
+        - name: plugin-dir
+          mountPath: /csi
+      volumes:
+      - name: kubelet-dir
+        hostPath:
+          path: /var/lib/kubelet
+          type: Directory
+      - name: plugin-dir
+        hostPath:
+          path: /var/lib/kubelet/plugins/local.csi.scylladb.com/
+          type: DirectoryOrCreate
+      - name: registration-dir
+        hostPath:
+          path: /var/lib/kubelet/plugins_registry/
+          type: Directory
+      - name: volumes-dir
+        hostPath:
+          path: /mnt/persistent-volumes
+          type: Directory
diff --git a/Engine/Source/Programs/UnrealCloudDDC/Bootstrapping/AWS/scylla-provisioner/storageclass_xfs.yaml b/Engine/Source/Programs/UnrealCloudDDC/Bootstrapping/AWS/scylla-provisioner/storageclass_xfs.yaml
new file mode 100644
--- /dev/null
+++ b/Engine/Source/Programs/UnrealCloudDDC/Bootstrapping/AWS/scylla-provisioner/storageclass_xfs.yaml
@@ -0,0 +1,8 @@
+kind: StorageClass
+apiVersion: storage.k8s.io/v1
+metadata:
+  name: scylladb-local-xfs
+provisioner: local.csi.scylladb.com
+volumeBindingMode: WaitForFirstConsumer
+parameters:
+  csi.storage.k8s.io/fstype: xfs
diff --git a/Engine/Source/Programs/UnrealCloudDDC/Bootstrapping/AWS/scylla.cluster.yaml b/Engine/Source/Programs/UnrealCloudDDC/Bootstrapping/AWS/scylla.cluster.yaml
new file mode 100644
--- /dev/null
+++ b/Engine/Source/Programs/UnrealCloudDDC/Bootstrapping/AWS/scylla.cluster.yaml
@@ -0,0 +1,64 @@
+serviceMonitor:
+  create: true
+
+network:
+  hostNetworking: true
+sysctls:
+- "fs.aio-max-nr=2097152"
+developerMode: false
+datacenter: "%REGION_REPLACE%"
+racks:
+- name: a
+  members: 1
+  storage:
+    storageClassName: scylladb-local-xfs
+    capacity: 800G
+  resources:
+    limits:
+      cpu: 3
+      memory: 28G
+  placement:
+    nodeAffinity:
+      requiredDuringSchedulingIgnoredDuringExecution:
+        nodeSelectorTerms:
+          - matchExpressions:
+            - key: topology.kubernetes.io/zone
+              operator: In
+              values:
+              - "%AZ_REPLACE_0%"
+            - key: scylla.scylladb.com/node-type
+              operator: In
+              values:
+              - scylla
+    tolerations:
+      - key: role
+        operator: Equal
+        value: scylla-clusters
+        effect: NoSchedule
+- name: b
+  members: 1
+  storage:
+    storageClassName: scylladb-local-xfs
+    capacity: 800G
+  resources:
+    limits:
+      cpu: 3
+      memory: 28G
+  placement:
+    nodeAffinity:
+      requiredDuringSchedulingIgnoredDuringExecution:
+        nodeSelectorTerms:
+        - matchExpressions:
+          - key: topology.kubernetes.io/zone
+            operator: In
+            values:
+            - "%AZ_REPLACE_1%"
+          - key: scylla.scylladb.com/node-type
+            operator: In
+            values:
+            - scylla
+    tolerations:
+    - key: role
+      operator: Equal
+      value: scylla-clusters
+      effect: NoSchedule
\ No newline at end of file
diff --git a/Engine/Source/Programs/UnrealCloudDDC/Bootstrapping/AWS/scylla.manager.yaml b/Engine/Source/Programs/UnrealCloudDDC/Bootstrapping/AWS/scylla.manager.yaml
new file mode 100644
--- /dev/null
+++ b/Engine/Source/Programs/UnrealCloudDDC/Bootstrapping/AWS/scylla.manager.yaml
@@ -0,0 +1,2 @@
+serviceMonitor:
+  create: true
diff --git a/Engine/Source/Programs/UnrealCloudDDC/Bootstrapping/AWS/setup.sh b/Engine/Source/Programs/UnrealCloudDDC/Bootstrapping/AWS/setup.sh
new file mode 100644
--- /dev/null
+++ b/Engine/Source/Programs/UnrealCloudDDC/Bootstrapping/AWS/setup.sh
@@ -0,0 +1,144 @@
+#! /bin/bash
+set -euo pipefail
+
+REGION_REPLACE=$AWS_REGION
+REGION_REPLACE_CLEAN=$(echo $AWS_REGION | sed -r 's/-/_/g')
+AZ_REPLACE_0=$(aws ec2 describe-availability-zones | jq -r .AvailabilityZones[0].ZoneName)
+AZ_REPLACE_1=$(aws ec2 describe-availability-zones | jq -r .AvailabilityZones[1].ZoneName)
+
+if [ -f "[REDACTED_USER_HOME]/cloud-ddc.conf" ] 
+then
+  source [REDACTED_USER_HOME]/cloud-ddc.conf
+else
+  S3BUCKET_NAME_SUFFIX=$(date | md5sum | cut -d ' ' -f 1 | cut -b 1-8 | tr '[:upper:]' '[:lower:]')
+  S3BUCKET_NAME=cloud-ddc-${S3BUCKET_NAME_SUFFIX}-${AWS_REGION}
+  SERVICE_ACCOUNT_TOKEN=$(date | md5sum | cut -d ' ' -f 1)
+
+  echo "S3BUCKET_NAME_SUFFIX=$S3BUCKET_NAME_SUFFIX
+S3BUCKET_NAME=$S3BUCKET_NAME
+SERVICE_ACCOUNT_TOKEN=$SERVICE_ACCOUNT_TOKEN
+" > [REDACTED_USER_HOME]/cloud-ddc.conf  
+fi
+
+echo "Deploying to region: \"$REGION_REPLACE\" using Availability zones: \"$AZ_REPLACE_0\" and \"$AZ_REPLACE_1\""
+
+echo "Updating configuration files..."
+sed -i "s/\%REGION_REPLACE\%/$REGION_REPLACE/g" eksctl.yaml
+sed -i "s/\%AZ_REPLACE_0\%/$AZ_REPLACE_0/g" eksctl.yaml
+sed -i "s/\%AZ_REPLACE_1\%/$AZ_REPLACE_1/g" eksctl.yaml
+
+sed -i "s/\%REGION_REPLACE\%/$REGION_REPLACE/g" cloud-ddc.yaml
+sed -i "s/\%REGION_REPLACE_CLEAN\%/$REGION_REPLACE_CLEAN/g" cloud-ddc.yaml
+sed -i "s/\%AZ_REPLACE_0\%/$AZ_REPLACE_0/g" cloud-ddc.yaml
+sed -i "s/\%AZ_REPLACE_1\%/$AZ_REPLACE_1/g" cloud-ddc.yaml
+sed -i "s/\%SERVICE_ACCOUNT_TOKEN\%/$SERVICE_ACCOUNT_TOKEN/g" cloud-ddc.yaml
+sed -i "s/\%S3BUCKET_NAME\%/$S3BUCKET_NAME/g" cloud-ddc.yaml
+
+sed -i "s/\%REGION_REPLACE\%/$REGION_REPLACE/g" scylla.cluster.yaml
+sed -i "s/\%AZ_REPLACE_0\%/$AZ_REPLACE_0/g" scylla.cluster.yaml
+sed -i "s/\%AZ_REPLACE_1\%/$AZ_REPLACE_1/g" scylla.cluster.yaml
+
+chmod +x ./delete.sh
+
+mkdir -p [REDACTED_USER_HOME]/bin
+
+echo "Installing eksctl..."
+curl --silent --location "[REDACTED_PROFILE_URL]/eksctl/releases/latest/download/eksctl_$(uname -s)_amd64.tar.gz" | tar xz -C /tmp && sudo mv /tmp/eksctl /usr/local/bin && sudo cp /usr/local/bin/eksctl [REDACTED_USER_HOME]/bin/eksctl
+
+echo "Installing helm..."
+curl -LO https://get.helm.sh/helm-v3.11.0-linux-amd64.tar.gz && tar xfz helm-v3.11.0-linux-amd64.tar.gz && sudo mv linux-amd64/helm /usr/local/bin/ && sudo cp /usr/local/bin/helm [REDACTED_USER_HOME]/bin/helm
+
+echo "Preparing AWS Resource Group for created resources"
+aws resource-groups create-group --name unreal-cloud-ddc --description "Resource group of all resources created by the UnrealCloudDDC bootstrapper" --resource-query file://resource-query.json
+
+if aws s3api head-bucket --bucket "$S3BUCKET_NAME" 2>/dev/null; then
+  echo "S3 Bucket $S3BUCKET_NAME already present ..."
+else
+  echo "Creating S3 bucket ..."
+  aws s3api create-bucket --bucket "$S3BUCKET_NAME" --region "$AWS_REGION" --acl private --create-bucket-configuration LocationConstraint=$AWS_REGION
+  aws s3api put-bucket-tagging --bucket "$S3BUCKET_NAME" --region "$AWS_REGION" --tagging "TagSet=[{Key=created-by,Value=CloudDDCBootstrapper}]"
+  
+  echo "Created S3 bucket $S3BUCKET_NAME..."
+fi
+
+echo "Creating EKS cluster..."
+eksctl create cluster -f ./eksctl.yaml
+
+#This works around a strange issue in eksctl were it can't have this managed node group included as part of the rest of the cluster definitions
+eksctl create nodegroup --cluster unreal-cloud-ddc --name monitoring-pool --node-type i3.large --nodes 1 --node-labels "pool=monitoring-pool"
+
+echo "Adding helm repos..."
+helm repo add eks https://aws.github.io/eks-charts
+helm repo add external-dns https://kubernetes-sigs.github.io/external-dns/
+helm repo add jetstack https://charts.jetstack.io
+helm repo add scylla https://scylla-operator-charts.storage.googleapis.com/stable
+helm repo update
+
+echo "Installing AWS Load balancer..."
+helm upgrade aws-load-balancer eks/aws-load-balancer-controller --install -n kube-system --wait --set serviceAccount.create=false,serviceAccount.name=aws-load-balancer-controller,clusterName=unreal-cloud-ddc
+echo "Installing External-DNS..."
+helm upgrade external-dns external-dns/external-dns --install -n kube-system --wait --set serviceAccount.create=false,serviceAccount.name=external-dns
+
+echo "Starting the cert manger..."
+helm install \
+  cert-manager jetstack/cert-manager \
+  --namespace cert-manager \
+  --create-namespace \
+  --version v1.13.1 \
+  --set installCRDs=true
+kubectl wait --for condition=established --timeout=60s crd/certificates.cert-manager.io crd/issuers.cert-manager.io
+kubectl -n cert-manager rollout status --timeout=5m deployment.apps/cert-manager-webhook
+
+echo "Installing prometheus..."
+helm install prometheus-operator oci://registry-1.docker.io/bitnamicharts/kube-prometheus --namespace prometheus-operator --create-namespace 
+
+echo "Starting the scylla operator..."
+helm install scylla-operator scylla/scylla-operator --create-namespace --namespace scylla-operator
+kubectl wait --for condition=established crd/nodeconfigs.scylla.scylladb.com
+kubectl wait --for condition=established crd/scyllaclusters.scylla.scylladb.com
+kubectl -n scylla-operator rollout status --timeout=5m deployment.apps/scylla-operator
+kubectl -n scylla-operator rollout status --timeout=5m deployment.apps/webhook-server
+
+# Configure nodes
+echo "Configuring scylla nodes..."
+kubectl apply --server-side -f nodeconfig-alpha.yaml
+kubectl -n local-csi-driver apply --server-side -f ./scylla-provisioner/local-csi-driver/
+kubectl apply --server-side -f ./scylla-provisioner/storageclass_xfs.yaml
+echo "Your disks are ready to use."
+
+echo "Installing scylla manager..."
+helm install scylla-manager scylla/scylla-manager --create-namespace --namespace scylla-manager -f scylla.manager.yaml
+
+echo "Installing ScyllaDB..."
+helm install scylla scylla/scylla --create-namespace --namespace scylla -f scylla.cluster.yaml --wait
+echo "Waiting for scylla nodes..."
+kubectl wait pods -l "app.kubernetes.io/name=scylla" -n scylla --for condition=Ready --timeout=10m
+
+echo "Configuring Scylla Monitoring..."
+kubectl -n scylla apply --server-side -f ./scylla-monitoring.yaml
+
+echo "Installing UnrealCloudDDC..."
+helm upgrade unreal-cloud-ddc oci://ghcr.io/epicgames/unreal-cloud-ddc --install --namespace unreal-cloud-ddc --values cloud-ddc.yaml --wait --version 0.5.2+helm
+
+# wait for the load balancer to get a host name
+DDC_HOST=""
+while [ -z $DDC_HOST ]; do
+  echo "Waiting for external IP"
+  DDC_HOST=$(kubectl get svc unreal-cloud-ddc --namespace unreal-cloud-ddc -o jsonpath="{.status.loadBalancer.ingress[*].hostname}")
+  [ -z "$DDC_HOST" ] && sleep 10
+done
+
+echo "CloudDDC Installed and ready to use"  
+echo " * Use \"$SERVICE_ACCOUNT_TOKEN\" as service account"
+echo " * Use \"$DDC_HOST\" as host url"
+
+echo "Please note that we strongly urge you to update this setup to use your SSO if you have one and to setup https, please see the documentation for more information."
+
+echo "To verify that you can upload data you can use this command: \"curl http://$DDC_HOST/api/v1/refs/ddc/default/00000000000000000000000000000000000000aa -X PUT --data 'test' -H 'content-type: application/octet-stream' -H 'X-Jupiter-IoHash: 4878CA0425C739FA427F7EDA20FE845F6B2E46BA' -i -H 'Authorization: ServiceAccount $SERVICE_ACCOUNT_TOKEN'\""
+echo "After uploading data you can fetch it using this command: \"curl http://$DDC_HOST/api/v1/refs/ddc/default/00000000000000000000000000000000000000aa.json -i -H 'Authorization: ServiceAccount $SERVICE_ACCOUNT_TOKEN'\""
+echo "Please note that it may take a few minutes for this host to be resolvable by DNS, if you encounter issues about resolving the host then please wait a while and try again. "
+echo "To run against this instance in your Unreal Editor you can add this to your DefaultEngine.ini - please note that this requires UE 5.4 to work"
+echo ""
+echo "[StorageServers]"
+echo "Default=(Host=\"$DDC_HOST\", Namespace=\"ddc\", AuthScheme=\"ServiceAccount\", OAuthAccess=\"$SERVICE_ACCOUNT_TOKEN\")"
+echo ""
\ No newline at end of file
