Unreal Cloud DDC - Added more tracing to metrics calculation to better understand how long time it takes for each part.

diff --git a/Engine/Source/Programs/UnrealCloudDDC/Jupiter/Implementation/Blob/ScyllaBlobIndex.cs b/Engine/Source/Programs/UnrealCloudDDC/Jupiter/Implementation/Blob/ScyllaBlobIndex.cs
--- a/Engine/Source/Programs/UnrealCloudDDC/Jupiter/Implementation/Blob/ScyllaBlobIndex.cs
+++ b/Engine/Source/Programs/UnrealCloudDDC/Jupiter/Implementation/Blob/ScyllaBlobIndex.cs
@@ -364,11 +364,12 @@ public class ScyllaBlobIndex : IBlobIndex
 
 	public async Task<BucketStats> CalculateBucketStatisticsAsync(NamespaceId ns, BucketId bucket)
 	{
-		using TelemetrySpan scope =  _tracer.BuildScyllaSpan("scylla.calc_bucket_stats");
-
 		string nsAsString = ns.ToString();
 		string bucketAsString = bucket.ToString();
 
+		using TelemetrySpan scope =  _tracer.BuildScyllaSpan("scylla.calc_bucket_stats")
+			.SetAttribute("resource.name", $"{nsAsString}.{bucketAsString}");
+
 		int totalCountOfRefs = 0;
 		int totalCountOfBlobs = 0;
 		long totalSizeOfBlobs = 0;
@@ -406,6 +407,8 @@ public class ScyllaBlobIndex : IBlobIndex
 
 			Task calcBlobStats = Task.Run(async () =>
 			{
+				using TelemetrySpan _ =  _tracer.BuildScyllaSpan("scylla.calc_bucket_stats_shard")
+					.SetAttribute("resource.name", $"{nsAsString}.{bucketAsString}.{hashPrefix}");
 				// blob stats
 				BoundStatement? boundStatement = _getBucketStatsStatement.Bind(nsAsString, bucketAsString, hashPrefix);
 
diff --git a/Engine/Source/Programs/UnrealCloudDDC/Jupiter/Implementation/MetricsService.cs b/Engine/Source/Programs/UnrealCloudDDC/Jupiter/Implementation/MetricsService.cs
--- a/Engine/Source/Programs/UnrealCloudDDC/Jupiter/Implementation/MetricsService.cs
+++ b/Engine/Source/Programs/UnrealCloudDDC/Jupiter/Implementation/MetricsService.cs
@@ -11,6 +11,7 @@ using Jupiter.Implementation.Blob;
 using Microsoft.Extensions.DependencyInjection;
 using Microsoft.Extensions.Logging;
 using Microsoft.Extensions.Options;
+using OpenTelemetry.Trace;
 
 namespace Jupiter.Implementation
 {
@@ -113,6 +114,7 @@ namespace Jupiter.Implementation
 		private readonly IBlobIndex _blobIndex;
 
 		private readonly ILogger _logger;
+		private readonly Tracer _tracer;
 		private readonly Gauge<double> _blobSizeAvgGauge;
 		private readonly Gauge<long> _blobSizeMinGauge;
 		private readonly Gauge<long> _blobSizeMaxGauge;
@@ -120,10 +122,11 @@ namespace Jupiter.Implementation
 		private readonly Gauge<long> _blobSizeCountGauge;
 		private readonly Gauge<long> _blobSizeTotalGauge;
 
-		public MetricsCalculator(IBlobIndex blobIndex, Meter meter, ILogger<MetricsService> logger)
+		public MetricsCalculator(IBlobIndex blobIndex, Meter meter, ILogger<MetricsService> logger, Tracer tracer)
 		{
 			_blobIndex = blobIndex;
 			_logger = logger;
+			_tracer = tracer;
 
 			_blobSizeAvgGauge = meter.CreateGauge<double>("blobstats.bucket_size.avg");
 			_blobSizeMinGauge = meter.CreateGauge<long>("blobstats.bucket_size.min");
@@ -135,6 +138,10 @@ namespace Jupiter.Implementation
 
 		public async Task<BucketStats?> CalculateStatsForBucketAsync(NamespaceId ns, BucketId bucket)
 		{
+			using TelemetrySpan removeBlobScope = _tracer.StartActiveSpan("metrics.calculate")
+				.SetAttribute("operation.name", "metrics.calculate")
+				.SetAttribute("resource.name", $"{ns}.{bucket}");
+
 			KeyValuePair<string, object?>[] tags = new[] { new KeyValuePair<string, object?>("Bucket", bucket.ToString()), new KeyValuePair<string, object?>("Namespace", ns.ToString()) };
 
 			BucketStats stats = await _blobIndex.CalculateBucketStatisticsAsync(ns, bucket);
