Unreal Cloud DDC - Clarified some points in the readme based on feedback from the docs team

diff --git a/Engine/Source/Programs/UnrealCloudDDC/README.md b/Engine/Source/Programs/UnrealCloudDDC/README.md
--- a/Engine/Source/Programs/UnrealCloudDDC/README.md
+++ b/Engine/Source/Programs/UnrealCloudDDC/README.md
@@ -12,7 +12,6 @@ Unreal Cloud DDC can signficantly help teams speed up their cook processes in th
 - [License](#license)
 - [Directories](#directories)
 - [Dependencies](#dependencies)
-- [Other useful things](#other-useful-things)
 - [Functional Test requirements](#functional-test-requirements)
 - [Running locally](#running-locally)
 - [Deployment](#deployment)
@@ -51,12 +50,8 @@ We do provide container images at [REDACTED_PROFILE_URL]/EpicGames/packages/con
 
 * DotNet Core 6 (and Visual Studio 2022 or VS Code)
 * Docker
-* Scylla
-* Blob storage (S3, Azure Blob Store or a local filesystem)
-
-# Other useful things
-* MongoDB
-* Minio
+* Database (Scylla is recommended, MongoDB is also supported for single regions)
+* Blob storage (S3, S3 emulations like Minio, Azure Blob Store or a local filesystem)
 * Docker Compose
 
 # Functional Test requirements
@@ -94,7 +89,7 @@ https://www.scylladb.com/download/#open-source
 Scylla provides machine images for use in cloud environments.
 
 ## AWS
-This is the most tested deployment form as this is how we operate it at Epic. The helm chart we install into each regions kubernetes cluster is provided in this repo.
+This is the most tested deployment form as this is how we operate it at Epic. The helm chart we install into each regions kubernetes cluster is available in the  `Helm` directory.
 
 ## On premise
 UnrealCloudDDC can be deployed onprem without using any cloud resources. You can either setup a Mongo database for this (if you only intend to run this in a single region) or Scylla if you inted to run it multi region but still on premise. If you are starting with one region but might expand later we recommend using Scylla - as that allows you to just scale out while Mongo would require drop all your existing state.
@@ -230,7 +225,7 @@ This is typically exposed on port `8008` and as `corp-http` within kubernetes.
 ## Internal Port
 The internal port is only needed to be reachable by other UnrealCloudDDC instances. This exposes everything that the private port does but also certain apis that are deemed sensitive (enumerating content via the replication log primarily).
 This is exposed on port `8080` and as `internal-http` within kubernetes.
-Its recommended to keep this ingress only to other UnrealCloudDDC instances via a private VPC or using some kind of ip-range allow list or similar.
+Its recommended to keep this port only to other UnrealCloudDDC instances via a private VPC or using some kind of ip-range allow list or similar.
 
 Note that this port is primarily used for the speculative blob replication (see `Blob replication setup`). 
 
