Unreal Cloud DDC - Updated documention in some places were it was out of date, clarifying how to run it on GCE and the state of the Azure support (in use at licensees)

diff --git a/Engine/Source/Programs/UnrealCloudDDC/README.md b/Engine/Source/Programs/UnrealCloudDDC/README.md
--- a/Engine/Source/Programs/UnrealCloudDDC/README.md
+++ b/Engine/Source/Programs/UnrealCloudDDC/README.md
@@ -19,6 +19,7 @@ Unreal Cloud DDC can signficantly help teams speed up their cook processes in th
   - [AWS](#aws)
   - [On premise](#on-premise)
   - [Azure](#azure)
+  - [GCE](#gce)
   - [Testing your deployment](#testing-your-deployment)
 - [Monitoring](#monitoring)
   - [Health Checks](#health-checks)
@@ -69,8 +70,9 @@ The AWS compose file can be replaced with `docker-compose-azure.yml` if you want
 Docker compose setups disable authentication for to make it quick to get started, generally we recommend that you hook UnrealCloudDDC up to a OIDC provider before deploying this.
 
 # Deployment
-UnrealCloudDDC is currently only run in production on AWS, but the requirements on storage and db are very generic and also abstracted. 
-We have basic (untested) support for Azure services.
+UnrealCloudDDC is currently only run in production on AWS for Epic, but the requirements on storage and db are very generic and also abstracted. 
+We have licensees running on Azure and include support for that but we have limited testing coverage for this.
+GCE can be used by using GCS with its S3 api which we have licensees that do, as with azure we have very limited testing for this mode.
 
 We provide helm values (under `/Helm`) that we use for epic internal deployments to kubernetes, but kubernetes is not a requirement.
 
@@ -97,6 +99,9 @@ UnrealCloudDDC can be deployed onprem without using any cloud resources. You can
 ## Azure
 To deploy on AWS you will just need to set Azure as your cloud provider and specify the `Azure.ConnectionString` setting with a connection string to your Azure Blob Storage.
 
+## GCE
+To run using GCS you will need to use the S3 api they provide as well as set the `S3.UseChunkEncoding` setting to `false`
+
 ## Testing your deployment
 Once you have a deployment up and running you can connect to the machine and run curl commands to verify its working as it should.
 
@@ -118,7 +123,7 @@ curl http://localhost/api/v1/refs/test-namespace/default/00000000000000000000000
 ```
 
 # Monitoring
-We use Datadog to monitor our services, as such UnrealCloudDDC is instrumented to work well with that service. But all logs are delivered as structured logs to stdout, so any monitoring service that understands structured logs should be able to monitor it quite well.
+We use Datadog to monitor our services, as such UnrealCloudDDC is instrumented to work well with that service. All logs are delivered as structured logs to stdout, so any monitoring service that understands structured logs should be able to monitor it quite well. Traces are output using OpenTelemetry formats so any monitoring service that can ingest that should be compatible.
 
 ## Health Checks
 All UnrealCloudDDC services use health checks to monitor themselves, any background services they may run and any dependent service they may have (DB / Blob store etc).
