Unreal Cloud DDC - Excluded batch endpoint from swagger as it struggle to understand how it works.

diff --git a/Engine/Source/Programs/UnrealCloudDDC/Jupiter/Controllers/ReferencesController.cs b/Engine/Source/Programs/UnrealCloudDDC/Jupiter/Controllers/ReferencesController.cs
--- a/Engine/Source/Programs/UnrealCloudDDC/Jupiter/Controllers/ReferencesController.cs
+++ b/Engine/Source/Programs/UnrealCloudDDC/Jupiter/Controllers/ReferencesController.cs
@@ -825,6 +825,7 @@ namespace Jupiter.Controllers
 		[HttpPost("{ns}")]
 		[Consumes(CustomMediaTypeNames.UnrealCompactBinary)]
 		[Produces(CustomMediaTypeNames.UnrealCompactBinary)]
+		[ApiExplorerSettings(IgnoreApi = true)]
 		public async Task<IActionResult> BatchAsync(
 			[FromRoute] [Required] NamespaceId ns,
 			[FromBody] [Required] BatchOps ops)
