Unreal Cloud DDC - Fixed failing tests - we can only early exit when we have tried all cutoff windows and determined there were now files.

diff --git a/Engine/Source/Programs/UnrealCloudDDC/Jupiter/Implementation/Blob/FileSystemStore.cs b/Engine/Source/Programs/UnrealCloudDDC/Jupiter/Implementation/Blob/FileSystemStore.cs
--- a/Engine/Source/Programs/UnrealCloudDDC/Jupiter/Implementation/Blob/FileSystemStore.cs
+++ b/Engine/Source/Programs/UnrealCloudDDC/Jupiter/Implementation/Blob/FileSystemStore.cs
@@ -240,11 +240,11 @@ namespace Jupiter.Implementation
 							// if the file was gced while running we can just ignore it
 						}
 					}
+				}
 
-					if (!hadFiles)
-					{
-						return countOfBlobsRemoved;
-					}
+				if (!hadFiles)
+				{
+					return countOfBlobsRemoved;
 				}
 			}
 
