UnrealCloudDDC: Fix static analysis warnings.

#rnx

diff --git a/Engine/Source/Programs/UnrealCloudDDC/Jupiter/Controllers/ReferencesController.cs b/Engine/Source/Programs/UnrealCloudDDC/Jupiter/Controllers/ReferencesController.cs
--- a/Engine/Source/Programs/UnrealCloudDDC/Jupiter/Controllers/ReferencesController.cs
+++ b/Engine/Source/Programs/UnrealCloudDDC/Jupiter/Controllers/ReferencesController.cs
@@ -238,7 +238,7 @@ namespace Jupiter.Controllers
 							byte[] blobMemory = await blob.Stream.ToByteArrayAsync();
 							CbObject cb = new CbObject(blobMemory);
 
-							IAsyncEnumerable<Attachment> attachments = _referenceResolver.GetAttachments(ns, cb);
+							IAsyncEnumerable<Attachment> attachments = _referenceResolver.GetAttachmentsAsync(ns, cb);
 
 							using CbPackageBuilder writer = new CbPackageBuilder();
 							writer.AddAttachment(objectRecord.BlobIdentifier.AsIoHash(), CbPackageAttachmentFlags.IsObject, blobMemory);
@@ -338,7 +338,7 @@ namespace Jupiter.Controllers
 								List<BlobId> referencedBlobs;
 								try
 								{
-									IAsyncEnumerable<BlobId> referencedBlobsEnumerable = _referenceResolver.GetReferencedBlobs(ns, cb);
+									IAsyncEnumerable<BlobId> referencedBlobsEnumerable = _referenceResolver.GetReferencedBlobsAsync(ns, cb);
 									referencedBlobs = await referencedBlobsEnumerable.ToListAsync();
 								}
 								catch (PartialReferenceResolveException)
@@ -428,8 +428,6 @@ namespace Jupiter.Controllers
 			}
 		}
 
-
-
 		/// <summary>
 		/// Returns the metadata about a ref key
 		/// </summary>
diff --git a/Engine/Source/Programs/UnrealCloudDDC/Jupiter/Implementation/Blob/FileSystemStore.cs b/Engine/Source/Programs/UnrealCloudDDC/Jupiter/Implementation/Blob/FileSystemStore.cs
--- a/Engine/Source/Programs/UnrealCloudDDC/Jupiter/Implementation/Blob/FileSystemStore.cs
+++ b/Engine/Source/Programs/UnrealCloudDDC/Jupiter/Implementation/Blob/FileSystemStore.cs
@@ -172,6 +172,8 @@ namespace Jupiter.Implementation
 		/// <returns></returns>
 		public async Task<ulong> CleanupInternalAsync(CancellationToken cancellationToken, int batchSize = 1_000_000)
 		{
+			_ = batchSize;
+
 			using TelemetrySpan scope = _tracer.StartActiveSpan("gc.filesystem")
 				.SetAttribute("operation.name", "gc.filesystem");
 
