UnrealCloudDDC: Add cancellation tokens to async methods in interfaces.

#rnx

diff --git a/Engine/Source/Programs/UnrealCloudDDC/Jupiter.Common/Implementation/BlobId.cs b/Engine/Source/Programs/UnrealCloudDDC/Jupiter.Common/Implementation/BlobId.cs
--- a/Engine/Source/Programs/UnrealCloudDDC/Jupiter.Common/Implementation/BlobId.cs
+++ b/Engine/Source/Programs/UnrealCloudDDC/Jupiter.Common/Implementation/BlobId.cs
@@ -7,6 +7,7 @@ using System.Globalization;
 using System.IO;
 using System.Text.Json;
 using System.Text.Json.Serialization;
+using System.Threading;
 using System.Threading.Tasks;
 using Blake3;
 using EpicGames.Core;
@@ -125,7 +126,7 @@ namespace Jupiter.Implementation
 			return new BlobId(testObjectHash.HashData);
 		}
 
-		public static async Task<BlobId> FromStreamAsync(Stream stream)
+		public static async Task<BlobId> FromStreamAsync(Stream stream, CancellationToken cancellationToken = default)
 		{
 			using Hasher hasher = Hasher.New();
 
@@ -134,11 +135,11 @@ namespace Jupiter.Implementation
 
 			try
 			{
-				int read = await stream.ReadAsync(buffer, 0, buffer.Length);
+				int read = await stream.ReadAsync(buffer, 0, buffer.Length, cancellationToken);
 				while (read > 0)
 				{
 					hasher.UpdateWithJoin(new ReadOnlySpan<byte>(buffer, 0, read));
-					read = await stream.ReadAsync(buffer, 0, buffer.Length);
+					read = await stream.ReadAsync(buffer, 0, buffer.Length, cancellationToken);
 				}
 				Hash blake3Hash = hasher.Finalize();
 
diff --git a/Engine/Source/Programs/UnrealCloudDDC/Jupiter.Common/Implementation/StreamingBodyPayload.cs b/Engine/Source/Programs/UnrealCloudDDC/Jupiter.Common/Implementation/StreamingBodyPayload.cs
--- a/Engine/Source/Programs/UnrealCloudDDC/Jupiter.Common/Implementation/StreamingBodyPayload.cs
+++ b/Engine/Source/Programs/UnrealCloudDDC/Jupiter.Common/Implementation/StreamingBodyPayload.cs
@@ -2,6 +2,7 @@
 
 using System;
 using System.IO;
+using System.Threading;
 using System.Threading.Tasks;
 using Jupiter.Utils;
 using Microsoft.AspNetCore.Http;
@@ -28,12 +29,12 @@ namespace Jupiter.Common.Implementation
 			_buffer = source;
 		}
 
-		public static async Task<MemoryBufferedPayload> CreateAsync(Tracer tracer, Stream s)
+		public static async Task<MemoryBufferedPayload> CreateAsync(Tracer tracer, Stream s, CancellationToken cancellationToken)
 		{
 			using TelemetrySpan scope = tracer.StartActiveSpan("payload.buffer")
 				.SetAttribute("operation.name", "payload.buffer")
 				.SetAttribute("bufferType", "Memory");
-			MemoryBufferedPayload payload = new MemoryBufferedPayload(await s.ToByteArrayAsync());
+			MemoryBufferedPayload payload = new MemoryBufferedPayload(await s.ToByteArrayAsync(cancellationToken));
 
 			return payload;
 		}
@@ -122,7 +123,7 @@ namespace Jupiter.Common.Implementation
 			_length = _tempFile.Length;
 		}
 
-		public static async Task<FilesystemBufferedPayload> CreateAsync(Tracer tracer, Stream s, string filesystemRoot)
+		public static async Task<FilesystemBufferedPayload> CreateAsync(Tracer tracer, Stream s, string filesystemRoot, CancellationToken cancellationToken)
 		{
 			FilesystemBufferedPayload payload = new FilesystemBufferedPayload(filesystemRoot);
 
@@ -131,7 +132,7 @@ namespace Jupiter.Common.Implementation
 					.SetAttribute("operation.name", "payload.buffer")
 					.SetAttribute("bufferType", "Filesystem");
 				await using FileStream fs = payload._tempFile.OpenWrite();
-				await s.CopyToAsync(fs);
+				await s.CopyToAsync(fs, cancellationToken);
 			}
 
 			payload._tempFile.Refresh();
@@ -182,7 +183,7 @@ namespace Jupiter.Common.Implementation
 			Directory.CreateDirectory(options.CurrentValue.FilesystemTempPayloadRoot);
 		}
 
-		public Task<IBufferedPayload> CreateFromRequestAsync(HttpRequest request)
+		public Task<IBufferedPayload> CreateFromRequestAsync(HttpRequest request, CancellationToken cancellationToken)
 		{
 			long? contentLength = request.ContentLength;
 
@@ -191,23 +192,23 @@ namespace Jupiter.Common.Implementation
 				throw new Exception("Expected content-length on all requests");
 			}
 
-			return CreateFromStreamAsync(request.Body, contentLength.Value);
+			return CreateFromStreamAsync(request.Body, contentLength.Value, cancellationToken);
 		}
 
-		public async Task<IBufferedPayload> CreateFromStreamAsync(Stream s, long contentLength)
+		public async Task<IBufferedPayload> CreateFromStreamAsync(Stream s, long contentLength, CancellationToken cancellationToken)
 		{
 			// blob is small enough to fit into memory we just read it as is
 			if (contentLength < _options.CurrentValue.MemoryBufferSize)
 			{
-				return await MemoryBufferedPayload.CreateAsync(_tracer, s);
+				return await MemoryBufferedPayload.CreateAsync(_tracer, s, cancellationToken);
 			}
 
-			return await FilesystemBufferedPayload.CreateAsync(_tracer, s, _options.CurrentValue.FilesystemTempPayloadRoot);
+			return await FilesystemBufferedPayload.CreateAsync(_tracer, s, _options.CurrentValue.FilesystemTempPayloadRoot, cancellationToken);
 		}
 
-		public async Task<IBufferedPayload> CreateFilesystemBufferedPayloadAsync(Stream s)
+		public async Task<IBufferedPayload> CreateFilesystemBufferedPayloadAsync(Stream s, CancellationToken cancellationToken)
 		{
-			return await FilesystemBufferedPayload.CreateAsync(_tracer, s, _options.CurrentValue.FilesystemTempPayloadRoot);
+			return await FilesystemBufferedPayload.CreateAsync(_tracer, s, _options.CurrentValue.FilesystemTempPayloadRoot, cancellationToken);
 		}
 
 		public FilesystemBufferedPayloadWriter CreateFilesystemBufferedPayloadWriter()
diff --git a/Engine/Source/Programs/UnrealCloudDDC/Jupiter.Common/Utils/StreamUtils.cs b/Engine/Source/Programs/UnrealCloudDDC/Jupiter.Common/Utils/StreamUtils.cs
--- a/Engine/Source/Programs/UnrealCloudDDC/Jupiter.Common/Utils/StreamUtils.cs
+++ b/Engine/Source/Programs/UnrealCloudDDC/Jupiter.Common/Utils/StreamUtils.cs
@@ -1,6 +1,7 @@
 ﻿// Copyright Epic Games, Inc. All Rights Reserved.
 
 using System.IO;
+using System.Threading;
 using System.Threading.Tasks;
 using EpicGames.AspNet;
 using Microsoft.AspNetCore.Http;
@@ -9,12 +10,12 @@ namespace Jupiter.Utils
 {
 	public static class StreamUtils
 	{
-		public static async Task<byte[]> ToByteArrayAsync(this Stream s)
+		public static async Task<byte[]> ToByteArrayAsync(this Stream s, CancellationToken cancellationToken)
 		{
 			try
 			{
 				await using MemoryStream ms = new MemoryStream();
-				await s.CopyToAsync(ms);
+				await s.CopyToAsync(ms, cancellationToken);
 				return ms.ToArray();
 			}
 			catch (BadHttpRequestException e)
diff --git a/Engine/Source/Programs/UnrealCloudDDC/Jupiter/Controllers/AdminController.cs b/Engine/Source/Programs/UnrealCloudDDC/Jupiter/Controllers/AdminController.cs
--- a/Engine/Source/Programs/UnrealCloudDDC/Jupiter/Controllers/AdminController.cs
+++ b/Engine/Source/Programs/UnrealCloudDDC/Jupiter/Controllers/AdminController.cs
@@ -49,7 +49,7 @@ namespace Jupiter.Controllers
 				return result;
 			}
 
-			List<(LastAccessRecord, DateTime)>? updatedRecords = await _lastAccessService.ProcessLastAccessRecordsAsync();
+			List<(LastAccessRecord, DateTime)>? updatedRecords = await _lastAccessService.ProcessLastAccessRecordsAsync(HttpContext.RequestAborted);
 
 			return Ok(new UpdatedRecordsResponse(
 				updatedRecords?.Select(tuple => new UpdatedRecordsResponse.UpdatedRecord(tuple.Item1, tuple.Item2))
diff --git a/Engine/Source/Programs/UnrealCloudDDC/Jupiter/Controllers/BlobController.cs b/Engine/Source/Programs/UnrealCloudDDC/Jupiter/Controllers/BlobController.cs
--- a/Engine/Source/Programs/UnrealCloudDDC/Jupiter/Controllers/BlobController.cs
+++ b/Engine/Source/Programs/UnrealCloudDDC/Jupiter/Controllers/BlobController.cs
@@ -190,7 +190,7 @@ namespace Jupiter.Controllers
 
 			try
 			{
-				Uri? uri = await _storage.MaybePutObjectWithRedirectAsync(ns, id);
+				Uri? uri = await _storage.MaybePutObjectWithRedirectAsync(ns, id, HttpContext.RequestAborted);
 				if (uri != null)
 				{
 					return Ok(new
@@ -199,9 +199,9 @@ namespace Jupiter.Controllers
 						RedirectUri = uri,
 					});
 				}
-				using IBufferedPayload payload = await _bufferedPayloadFactory.CreateFromRequestAsync(Request);
+				using IBufferedPayload payload = await _bufferedPayloadFactory.CreateFromRequestAsync(Request, HttpContext.RequestAborted);
 
-				BlobId identifier = await _storage.PutObjectAsync(ns, payload, id);
+				BlobId identifier = await _storage.PutObjectAsync(ns, payload, id, HttpContext.RequestAborted);
 				return Ok(new
 				{
 					Identifier = identifier.ToString()
@@ -232,12 +232,12 @@ namespace Jupiter.Controllers
 			_diagnosticContext.Set("Content-Length", Request.ContentLength ?? -1);
 			try
 			{
-				using IBufferedPayload payload = await _bufferedPayloadFactory.CreateFromRequestAsync(Request);
+				using IBufferedPayload payload = await _bufferedPayloadFactory.CreateFromRequestAsync(Request, HttpContext.RequestAborted);
 
 				await using Stream stream = payload.GetStream();
 
-				BlobId id = await BlobId.FromStreamAsync(stream);
-				await _storage.PutObjectKnownHashAsync(ns, payload, id);
+				BlobId id = await BlobId.FromStreamAsync(stream, HttpContext.RequestAborted);
+				await _storage.PutObjectKnownHashAsync(ns, payload, id, HttpContext.RequestAborted);
 
 				return Ok(new
 				{
@@ -276,14 +276,14 @@ namespace Jupiter.Controllers
 				return result;
 			}
 
-			await _storage.DeleteNamespaceAsync(ns);
+			await _storage.DeleteNamespaceAsync(ns, HttpContext.RequestAborted);
 
 			return NoContent();
 		}
 
 		private async Task DeleteImplAsync(NamespaceId ns, BlobId id)
 		{
-			await _storage.DeleteObjectAsync(ns, id);
+			await _storage.DeleteObjectAsync(ns, id, HttpContext.RequestAborted);
 		}
 
 		// ReSharper disable UnusedAutoPropertyAccessor.Global
@@ -401,7 +401,7 @@ namespace Jupiter.Controllers
 							}
 
 							using MemoryBufferedPayload payload = new MemoryBufferedPayload(op.Content);
-							tasks[index] = _storage.PutObjectAsync(op.Namespace.Value, payload, op.Id).ContinueWith((t, _) => (object?)t.Result, null, TaskScheduler.Current);
+							tasks[index] = _storage.PutObjectAsync(op.Namespace.Value, payload, op.Id, HttpContext.RequestAborted).ContinueWith((t, _) => (object?)t.Result, null, TaskScheduler.Current);
 							break;
 						}
 					case BatchOp.Operation.DELETE:
diff --git a/Engine/Source/Programs/UnrealCloudDDC/Jupiter/Controllers/CompressedBlobController.cs b/Engine/Source/Programs/UnrealCloudDDC/Jupiter/Controllers/CompressedBlobController.cs
--- a/Engine/Source/Programs/UnrealCloudDDC/Jupiter/Controllers/CompressedBlobController.cs
+++ b/Engine/Source/Programs/UnrealCloudDDC/Jupiter/Controllers/CompressedBlobController.cs
@@ -6,6 +6,7 @@ using System.ComponentModel.DataAnnotations;
 using System.Linq;
 using System.Net;
 using System.Net.Mime;
+using System.Threading;
 using System.Threading.Tasks;
 using EpicGames.AspNet;
 using EpicGames.Horde.Storage;
@@ -228,9 +229,9 @@ namespace Jupiter.Controllers
 
 			try
 			{
-				using IBufferedPayload payload = await _bufferedPayloadFactory.CreateFromRequestAsync(Request);
+				using IBufferedPayload payload = await _bufferedPayloadFactory.CreateFromRequestAsync(Request, HttpContext.RequestAborted);
 
-				ContentId identifier = await _storage.PutCompressedObjectAsync(ns, payload, id, HttpContext.RequestServices);
+				ContentId identifier = await _storage.PutCompressedObjectAsync(ns, payload, id, HttpContext.RequestServices, HttpContext.RequestAborted);
 
 				return Ok(new { Identifier = identifier.ToString() });
 			}
@@ -254,6 +255,8 @@ namespace Jupiter.Controllers
 		public async Task<IActionResult> PostAsync(
 			[Required] NamespaceId ns)
 		{
+			CancellationToken cancellationToken = HttpContext.RequestAborted;
+
 			ActionResult? result = await _requestHelper.HasAccessToNamespaceAsync(User, Request, ns, new[] { JupiterAclAction.WriteObject });
 			if (result != null)
 			{
@@ -264,9 +267,9 @@ namespace Jupiter.Controllers
 
 			try
 			{
-				using IBufferedPayload payload = await _bufferedPayloadFactory.CreateFromRequestAsync(Request);
+				using IBufferedPayload payload = await _bufferedPayloadFactory.CreateFromRequestAsync(Request, cancellationToken);
 
-				ContentId identifier = await _storage.PutCompressedObjectAsync(ns, payload, null, HttpContext.RequestServices);
+				ContentId identifier = await _storage.PutCompressedObjectAsync(ns, payload, null, HttpContext.RequestServices, cancellationToken);
 
 				return Ok(new
 				{
diff --git a/Engine/Source/Programs/UnrealCloudDDC/Jupiter/Controllers/ObjectController.cs b/Engine/Source/Programs/UnrealCloudDDC/Jupiter/Controllers/ObjectController.cs
--- a/Engine/Source/Programs/UnrealCloudDDC/Jupiter/Controllers/ObjectController.cs
+++ b/Engine/Source/Programs/UnrealCloudDDC/Jupiter/Controllers/ObjectController.cs
@@ -160,9 +160,9 @@ namespace Jupiter.Controllers
 			_diagnosticContext.Set("Content-Length", Request.ContentLength ?? -1);
 			try
 			{
-				using IBufferedPayload payload = await _bufferedPayloadFactory.CreateFromRequestAsync(Request);
+				using IBufferedPayload payload = await _bufferedPayloadFactory.CreateFromRequestAsync(Request, HttpContext.RequestAborted);
 
-				BlobId identifier = await _storage.PutObjectAsync(ns, payload, id);
+				BlobId identifier = await _storage.PutObjectAsync(ns, payload, id, HttpContext.RequestAborted);
 				return Ok(new PutBlobResponse(identifier));
 			}
 			catch (ClientSendSlowException e)
@@ -192,7 +192,7 @@ namespace Jupiter.Controllers
 				return NotFound(new ValidationProblemDetails { Title = $"Object {e.Blob} not found" });
 			}
 
-			byte[] blobContents = await blob.Stream.ToByteArrayAsync();
+			byte[] blobContents = await blob.Stream.ToByteArrayAsync(HttpContext.RequestAborted);
 			if (blobContents.Length == 0)
 			{
 				_logger.LogWarning("0 byte object found for {Id} {Namespace}", id, ns);
@@ -234,7 +234,7 @@ namespace Jupiter.Controllers
 				return result;
 			}
 
-			await _storage.DeleteObjectAsync(ns, id);
+			await _storage.DeleteObjectAsync(ns, id, HttpContext.RequestAborted);
 
 			return Ok(new DeletedResponse
 			{
@@ -252,7 +252,7 @@ namespace Jupiter.Controllers
 				return result;
 			}
 
-			await _storage.DeleteNamespaceAsync(ns);
+			await _storage.DeleteNamespaceAsync(ns, HttpContext.RequestAborted);
 
 			return Ok();
 		}
diff --git a/Engine/Source/Programs/UnrealCloudDDC/Jupiter/Controllers/ReferencesController.cs b/Engine/Source/Programs/UnrealCloudDDC/Jupiter/Controllers/ReferencesController.cs
--- a/Engine/Source/Programs/UnrealCloudDDC/Jupiter/Controllers/ReferencesController.cs
+++ b/Engine/Source/Programs/UnrealCloudDDC/Jupiter/Controllers/ReferencesController.cs
@@ -80,7 +80,7 @@ namespace Jupiter.Controllers
 		[ProducesResponseType(type: typeof(ProblemDetails), 400)]
 		public async Task<IActionResult> GetNamespacesAsync()
 		{
-			NamespaceId[] namespaces = await _refService.GetNamespacesAsync().ToArrayAsync();
+			NamespaceId[] namespaces = await _refService.GetNamespacesAsync(HttpContext.RequestAborted).ToArrayAsync(HttpContext.RequestAborted);
 
 			// filter namespaces down to only the namespaces the user has access to
 			List<NamespaceId> namespacesWithAccess = new();
@@ -174,7 +174,7 @@ namespace Jupiter.Controllers
 						}
 					case MediaTypeNames.Application.Octet:
 						{
-							byte[] blobMemory = await blob.Stream.ToByteArrayAsync();
+							byte[] blobMemory = await blob.Stream.ToByteArrayAsync(HttpContext.RequestAborted);
 							CbObject cb = new CbObject(blobMemory);
 
 							(int, CbField?) CountFields(CbObject o)
@@ -223,7 +223,7 @@ namespace Jupiter.Controllers
 							byte[] blobMemory;
 							{
 								using TelemetrySpan scope = _tracer.StartActiveSpan("json.readblob").SetAttribute("operation.name", "json.readblob");
-								blobMemory = await blob.Stream.ToByteArrayAsync();
+								blobMemory = await blob.Stream.ToByteArrayAsync(HttpContext.RequestAborted);
 							}
 							CbObject cb = new CbObject(blobMemory);
 							string s = cb.ToJson();
@@ -235,10 +235,10 @@ namespace Jupiter.Controllers
 					case CustomMediaTypeNames.UnrealCompactBinaryPackage:
 						{
 							using TelemetrySpan packageScope = _tracer.StartActiveSpan("cbpackage.fetch").SetAttribute("operation.name", "cbpackage.fetch");
-							byte[] blobMemory = await blob.Stream.ToByteArrayAsync();
+							byte[] blobMemory = await blob.Stream.ToByteArrayAsync(HttpContext.RequestAborted);
 							CbObject cb = new CbObject(blobMemory);
 
-							IAsyncEnumerable<Attachment> attachments = _referenceResolver.GetAttachmentsAsync(ns, cb);
+							IAsyncEnumerable<Attachment> attachments = _referenceResolver.GetAttachmentsAsync(ns, cb, HttpContext.RequestAborted);
 
 							using CbPackageBuilder writer = new CbPackageBuilder();
 							writer.AddAttachment(objectRecord.BlobIdentifier.AsIoHash(), CbPackageAttachmentFlags.IsObject, blobMemory);
@@ -254,19 +254,19 @@ namespace Jupiter.Controllers
 									if (attachment is BlobAttachment blobAttachment)
 									{
 										BlobId referencedBlob = blobAttachment.Identifier;
-										attachmentContents = await _blobStore.GetObjectAsync(ns, referencedBlob);
+										attachmentContents = await _blobStore.GetObjectAsync(ns, referencedBlob, cancellationToken: HttpContext.RequestAborted);
 									}
 									else if (attachment is ObjectAttachment objectAttachment)
 									{
 										flags |= CbPackageAttachmentFlags.IsObject;
 										BlobId referencedBlob = objectAttachment.Identifier;
-										attachmentContents = await _blobStore.GetObjectAsync(ns, referencedBlob);
+										attachmentContents = await _blobStore.GetObjectAsync(ns, referencedBlob, cancellationToken: HttpContext.RequestAborted);
 									}
 									else if (attachment is ContentIdAttachment contentIdAttachment)
 									{
 
 										ContentId contentId = contentIdAttachment.Identifier;
-										(attachmentContents, string mime) = await _blobStore.GetCompressedObjectAsync(ns, contentId, HttpContext.RequestServices);
+										(attachmentContents, string mime) = await _blobStore.GetCompressedObjectAsync(ns, contentId, HttpContext.RequestServices, cancellationToken: HttpContext.RequestAborted);
 										if (mime == CustomMediaTypeNames.UnrealCompressedBuffer)
 										{
 											flags |= CbPackageAttachmentFlags.IsCompressed;
@@ -305,7 +305,7 @@ namespace Jupiter.Controllers
 						}
 					case CustomMediaTypeNames.JupiterInlinedPayload:
 						{
-							byte[] blobMemory = await blob.Stream.ToByteArrayAsync();
+							byte[] blobMemory = await blob.Stream.ToByteArrayAsync(HttpContext.RequestAborted);
 							CbObject cb = new CbObject(blobMemory);
 
 							static (int, int) CountFields(CbObject o)
@@ -503,7 +503,7 @@ namespace Jupiter.Controllers
 
 				// we have to verify the blobs are available locally, as the record of the key is replicated a head of the content
 				// TODO: Once we support inline replication this step is not needed as at least one region as this blob, just maybe not this current one
-				byte[] blobContents = await blob.Stream.ToByteArrayAsync();
+				byte[] blobContents = await blob.Stream.ToByteArrayAsync(HttpContext.RequestAborted);
 				CbObject compactBinaryObject = new CbObject(blobContents);
 				// the reference resolver will throw if any blob is missing, so no need to do anything other then process each reference
 				IAsyncEnumerable<BlobId> references = _referenceResolver.GetReferencedBlobsAsync(ns, compactBinaryObject);
@@ -511,7 +511,7 @@ namespace Jupiter.Controllers
 
 				// we have to verify the blobs are available locally, as the record of the key is replicated a head of the content
 				// TODO: Once we support inline replication this step is not needed as at least one region as this blob, just maybe not this current one
-				BlobId[] unknownBlobs = await _blobStore.FilterOutKnownBlobsAsync(ns, new BlobId[] { record.BlobIdentifier });
+				BlobId[] unknownBlobs = await _blobStore.FilterOutKnownBlobsAsync(ns, new BlobId[] { record.BlobIdentifier }, HttpContext.RequestAborted);
 				if (unknownBlobs.Length != 0)
 				{
 					return NotFound(new ProblemDetails { Title = $"Object {bucket} {key} in namespace {ns} had at least one missing blob." });
@@ -581,7 +581,7 @@ namespace Jupiter.Controllers
 
 					// we have to verify the blobs are available locally, as the record of the key is replicated a head of the content
 					// TODO: Once we support inline replication this step is not needed as at least one region as this blob, just maybe not this current one
-					byte[] blobContents = await blob.Stream.ToByteArrayAsync();
+					byte[] blobContents = await blob.Stream.ToByteArrayAsync(HttpContext.RequestAborted);
 					CbObject cb = new CbObject(blobContents);
 					// the reference resolver will throw if any blob is missing, so no need to do anything other then process each reference
 					IAsyncEnumerable<BlobId> references = _referenceResolver.GetReferencedBlobsAsync(ns, cb);
@@ -629,7 +629,7 @@ namespace Jupiter.Controllers
 
 			try
 			{
-				using IBufferedPayload payload = await _bufferedPayloadFactory.CreateFromRequestAsync(Request);
+				using IBufferedPayload payload = await _bufferedPayloadFactory.CreateFromRequestAsync(Request, HttpContext.RequestAborted);
 
 				BlobId headerHash;
 				if (Request.Headers.TryGetValue(CommonHeaders.HashHeaderName, out StringValues headers))
@@ -662,7 +662,7 @@ namespace Jupiter.Controllers
 						{
 							// TODO: define a scheme for how a json object specifies references
 
-							blobHeader = await _blobStore.PutObjectAsync(ns, payload, headerHash);
+							blobHeader = await _blobStore.PutObjectAsync(ns, payload, headerHash, HttpContext.RequestAborted);
 
 							// TODO: convert the json object into a compact binary instead
 							CbWriter writer = new CbWriter();
@@ -685,7 +685,7 @@ namespace Jupiter.Controllers
 						}
 					case MediaTypeNames.Application.Octet:
 						{
-							blobHeader = await _blobStore.PutObjectAsync(ns, payload, headerHash);
+							blobHeader = await _blobStore.PutObjectAsync(ns, payload, headerHash, HttpContext.RequestAborted);
 
 							CbWriter writer = new CbWriter();
 							writer.BeginObject();
@@ -714,7 +714,7 @@ namespace Jupiter.Controllers
 				return Problem(e.Message, null, (int)HttpStatusCode.RequestTimeout);
 			}
 
-			(ContentId[] missingReferences, BlobId[] missingBlobs) = await _refService.PutAsync(ns, bucket, key, blobHeader, payloadObject);
+			(ContentId[] missingReferences, BlobId[] missingBlobs) = await _refService.PutAsync(ns, bucket, key, blobHeader, payloadObject, HttpContext.RequestAborted);
 
 			{
 				using TelemetrySpan scope = _tracer.StartActiveSpan("ref.put").SetAttribute("operation.name", "ref.put");
@@ -762,11 +762,11 @@ namespace Jupiter.Controllers
 #pragma warning disable CA2000 // Dispose objects before losing scope
 						using MemoryBufferedPayload payload = new MemoryBufferedPayload(blob);
 #pragma warning restore CA2000 // Dispose objects before losing scope
-						await _blobStore.PutCompressedObjectAsync(ns, payload, ContentId.FromIoHash(entry.AttachmentHash), HttpContext.RequestServices);
+						await _blobStore.PutCompressedObjectAsync(ns, payload, ContentId.FromIoHash(entry.AttachmentHash), HttpContext.RequestServices, HttpContext.RequestAborted);
 					}
 					else
 					{
-						await _blobStore.PutObjectAsync(ns, blob, BlobId.FromIoHash(entry.AttachmentHash));
+						await _blobStore.PutObjectAsync(ns, blob, BlobId.FromIoHash(entry.AttachmentHash), HttpContext.RequestAborted);
 					}
 				}
 			}
@@ -781,7 +781,7 @@ namespace Jupiter.Controllers
 			CbObject rootObject = packageReader.RootObject;
 			BlobId rootObjectHash = BlobId.FromIoHash(packageReader.RootHash);
 
-			(ContentId[] missingReferences, BlobId[] missingBlobs) = await _refService.PutAsync(ns, bucket, key, rootObjectHash, rootObject);
+			(ContentId[] missingReferences, BlobId[] missingBlobs) = await _refService.PutAsync(ns, bucket, key, rootObjectHash, rootObject, HttpContext.RequestAborted);
 
 			List<ContentHash> missingHashes = new List<ContentHash>(missingReferences);
 			missingHashes.AddRange(missingBlobs);
@@ -803,7 +803,7 @@ namespace Jupiter.Controllers
 
 			try
 			{
-				(ContentId[] missingReferences, BlobId[] missingBlobs) = await _refService.FinalizeAsync(ns, bucket, key, hash);
+				(ContentId[] missingReferences, BlobId[] missingBlobs) = await _refService.FinalizeAsync(ns, bucket, key, hash, HttpContext.RequestAborted);
 				List<ContentHash> missingHashes = new List<ContentHash>(missingReferences);
 				missingHashes.AddRange(missingBlobs);
 
@@ -877,7 +877,7 @@ namespace Jupiter.Controllers
 						throw new Exception();
 					}
 
-					CbObject cb = new CbObject(await blob.Stream.ToByteArrayAsync());
+					CbObject cb = new CbObject(await blob.Stream.ToByteArrayAsync(HttpContext.RequestAborted));
 
 					if (op.ResolveAttachments ?? false)
 					{
@@ -914,7 +914,7 @@ namespace Jupiter.Controllers
 
 					if (op.ResolveAttachments ?? false)
 					{
-						byte[] blobContents = await blob.Stream.ToByteArrayAsync();
+						byte[] blobContents = await blob.Stream.ToByteArrayAsync(HttpContext.RequestAborted);
 						CbObject cb = new CbObject(blobContents);
 						// the reference resolver will throw if any blob is missing, so no need to do anything other then process each reference
 						IAsyncEnumerable<BlobId> references = _referenceResolver.GetReferencedBlobsAsync(ns, cb);
@@ -959,7 +959,7 @@ namespace Jupiter.Controllers
 						throw new HashMismatchException(headerHash, objectHash);
 					}
 
-					(ContentId[] missingReferences, BlobId[] missingBlobs) = await _refService.PutAsync(ns, op.Bucket, op.Key, objectHash, op.Payload);
+					(ContentId[] missingReferences, BlobId[] missingBlobs) = await _refService.PutAsync(ns, op.Bucket, op.Key, objectHash, op.Payload, HttpContext.RequestAborted);
 					List<ContentHash> missingHashes = new List<ContentHash>(missingReferences);
 
 					return (CbSerializer.Serialize(new PutObjectResponse(missingHashes.ToArray())), HttpStatusCode.OK);
@@ -1048,7 +1048,7 @@ namespace Jupiter.Controllers
 
 			try
 			{
-				await _refService.DropNamespaceAsync(ns);
+				await _refService.DropNamespaceAsync(ns, HttpContext.RequestAborted);
 			}
 			catch (NamespaceNotFoundException e)
 			{
@@ -1078,7 +1078,7 @@ namespace Jupiter.Controllers
 			long countOfDeletedRecords;
 			try
 			{
-				countOfDeletedRecords = await _refService.DeleteBucketAsync(ns, bucket);
+				countOfDeletedRecords = await _refService.DeleteBucketAsync(ns, bucket, HttpContext.RequestAborted);
 			}
 			catch (NamespaceNotFoundException e)
 			{
@@ -1110,7 +1110,7 @@ namespace Jupiter.Controllers
 
 			try
 			{
-				bool deleted = await _refService.DeleteAsync(ns, bucket, key);
+				bool deleted = await _refService.DeleteAsync(ns, bucket, key, HttpContext.RequestAborted);
 				return Ok(new RefDeletedResponse(deleted ? 1 : 0));
 			}
 			catch (NamespaceNotFoundException e)
diff --git a/Engine/Source/Programs/UnrealCloudDDC/Jupiter/Controllers/SymbolsController.cs b/Engine/Source/Programs/UnrealCloudDDC/Jupiter/Controllers/SymbolsController.cs
--- a/Engine/Source/Programs/UnrealCloudDDC/Jupiter/Controllers/SymbolsController.cs
+++ b/Engine/Source/Programs/UnrealCloudDDC/Jupiter/Controllers/SymbolsController.cs
@@ -92,7 +92,7 @@ namespace Jupiter.Controllers
 				return BadRequest("No Blob Contents found");
 			}
 
-			byte[] blobMemory = await refContents.Stream.ToByteArrayAsync();
+			byte[] blobMemory = await refContents.Stream.ToByteArrayAsync(HttpContext.RequestAborted);
 			CbObject cb = new CbObject(blobMemory);
 			IoHash payloadHash = cb["pdbPayload"].AsBinaryAttachment().Hash;
 
@@ -146,7 +146,7 @@ namespace Jupiter.Controllers
 				case MediaTypeNames.Application.Octet:
 					{
 						CompressedBufferUtils utils = new CompressedBufferUtils(_tracer, _bufferedPayloadFactory);
-						using IBufferedPayload payload = await utils.DecompressContentAsync(referencedBlobContents.Stream, (ulong)referencedBlobContents.Length);
+						using IBufferedPayload payload = await utils.DecompressContentAsync(referencedBlobContents.Stream, (ulong)referencedBlobContents.Length, HttpContext.RequestAborted);
 						await using Stream s = payload.GetStream();
 						await using BlobContents contents = new BlobContents(s, s.Length);
 						await WriteBody(contents, MediaTypeNames.Application.Octet);
@@ -177,23 +177,23 @@ namespace Jupiter.Controllers
 			_diagnosticContext.Set("Content-Length", Request.ContentLength ?? -1);
 			CompressedBufferUtils utils = new CompressedBufferUtils(_tracer, _bufferedPayloadFactory);
 
-			IBufferedPayload payloadToUse = await _bufferedPayloadFactory.CreateFromRequestAsync(Request);
+			IBufferedPayload payloadToUse = await _bufferedPayloadFactory.CreateFromRequestAsync(Request, HttpContext.RequestAborted);
 			if (Request.ContentType == MediaTypeNames.Application.Octet)
 			{
 				await using Stream s = payloadToUse.GetStream();
 				using MemoryStream compressedStream = new MemoryStream();
 
 				// compress the content so we can use our normal path for processing content
-				utils.CompressContent(compressedStream, OoodleCompressorMethod.Kraken, OoodleCompressionLevel.VeryFast, await s.ToByteArrayAsync());
+				utils.CompressContent(compressedStream, OoodleCompressorMethod.Kraken, OoodleCompressionLevel.VeryFast, await s.ToByteArrayAsync(HttpContext.RequestAborted));
 
 				compressedStream.Seek(0, SeekOrigin.Begin);
-				payloadToUse = await _bufferedPayloadFactory.CreateFromStreamAsync(compressedStream, compressedStream.Length);
+				payloadToUse = await _bufferedPayloadFactory.CreateFromStreamAsync(compressedStream, compressedStream.Length, HttpContext.RequestAborted);
 			}
 
 			using IBufferedPayload payload = payloadToUse;
 			await using Stream hashStream = payload.GetStream();
-			BlobId attachmentHash = await BlobId.FromStreamAsync(hashStream);
-			IBufferedPayload decompressedContent = await utils.DecompressContentAsync(payload.GetStream(), (ulong)payload.Length);
+			BlobId attachmentHash = await BlobId.FromStreamAsync(hashStream, HttpContext.RequestAborted);
+			IBufferedPayload decompressedContent = await utils.DecompressContentAsync(payload.GetStream(), (ulong)payload.Length, HttpContext.RequestAborted);
 
 			(string pdbIdentifier, int pdbAge) = ExtractModuleInformation(moduleName, decompressedContent);
 
@@ -210,9 +210,9 @@ namespace Jupiter.Controllers
 			byte[] blob = writer.ToByteArray();
 			CbObject o = new CbObject(blob);
 			BlobId blobHeader = BlobId.FromBlob(blob);
-			await _blobStore.PutObjectKnownHashAsync(ns, payload, attachmentHash);
+			await _blobStore.PutObjectKnownHashAsync(ns, payload, attachmentHash, HttpContext.RequestAborted);
 
-			(ContentId[], BlobId[]) missingHashes = await _refService.PutAsync(ns, new BucketId(moduleName), RefId.FromName($"{moduleName}.{pdbIdentifier}{pdbAge}.{filename}"), blobHeader, o);
+			(ContentId[], BlobId[]) missingHashes = await _refService.PutAsync(ns, new BucketId(moduleName), RefId.FromName($"{moduleName}.{pdbIdentifier}{pdbAge}.{filename}"), blobHeader, o, HttpContext.RequestAborted);
 
 			if (missingHashes.Item1.Any() || missingHashes.Item2.Any())
 			{
diff --git a/Engine/Source/Programs/UnrealCloudDDC/Jupiter/Implementation/Blob/BlobIndexConsistencyCheckService.cs b/Engine/Source/Programs/UnrealCloudDDC/Jupiter/Implementation/Blob/BlobIndexConsistencyCheckService.cs
--- a/Engine/Source/Programs/UnrealCloudDDC/Jupiter/Implementation/Blob/BlobIndexConsistencyCheckService.cs
+++ b/Engine/Source/Programs/UnrealCloudDDC/Jupiter/Implementation/Blob/BlobIndexConsistencyCheckService.cs
@@ -60,18 +60,19 @@ namespace Jupiter.Implementation
 				return false;
 			}
 
-			await RunConsistencyCheckAsync();
+			await RunConsistencyCheckAsync(cancellationToken);
 
 			return true;
 		}
 
-		private async Task RunConsistencyCheckAsync()
+		private async Task RunConsistencyCheckAsync(CancellationToken cancellationToken)
 		{
 			ulong countOfBlobsChecked = 0;
 			ulong countOfIncorrectBlobsFound = 0;
 			string currentRegion = _jupiterSettings.CurrentValue.CurrentSite;
-			await Parallel.ForEachAsync(_blobIndex.GetAllBlobsAsync(), new ParallelOptions
+			await Parallel.ForEachAsync(_blobIndex.GetAllBlobsAsync(cancellationToken), new ParallelOptions
 			{
+				CancellationToken = cancellationToken,
 				MaxDegreeOfParallelism = _settings.CurrentValue.BlobIndexMaxParallelOperations,
 			},
 				async (tuple, token) =>
@@ -88,12 +89,12 @@ namespace Jupiter.Implementation
 
 					bool issueFound = false;
 					bool deleted = false;
-					List<string> regions = await _blobIndex.GetBlobRegionsAsync(ns, blobIdentifier);
+					List<string> regions = await _blobIndex.GetBlobRegionsAsync(ns, blobIdentifier, cancellationToken);
 					try
 					{
 						if (regions.Contains(currentRegion))
 						{
-							if (!await _blobService.ExistsInRootStoreAsync(ns, blobIdentifier))
+							if (!await _blobService.ExistsInRootStoreAsync(ns, blobIdentifier, token))
 							{
 								Interlocked.Increment(ref countOfIncorrectBlobsFound);
 								issueFound = true;
@@ -104,7 +105,7 @@ namespace Jupiter.Implementation
 
 									try
 									{
-										BlobContents _ = await _blobService.ReplicateObjectAsync(ns, blobIdentifier, force: true);
+										BlobContents _ = await _blobService.ReplicateObjectAsync(ns, blobIdentifier, force: true, token);
 									}
 									catch (BlobReplicationException e)
 									{
@@ -112,7 +113,7 @@ namespace Jupiter.Implementation
 										if (_settings.CurrentValue.AllowDeletesInBlobIndex)
 										{
 											_logger.LogWarning("Updating blob index to remove Blob {Blob} in namespace {Namespace} as we failed to repair it.", blobIdentifier, ns);
-											await _blobIndex.RemoveBlobFromRegionAsync(ns, blobIdentifier);
+											await _blobIndex.RemoveBlobFromRegionAsync(ns, blobIdentifier, cancellationToken: token);
 											deleted = true;
 										}
 										else
@@ -135,7 +136,7 @@ namespace Jupiter.Implementation
 										// this blob can not be repaired so we just delete it from the blob index
 										_logger.LogWarning("Blob {Blob} in namespace {Namespace} can not be repaired so removing existence from current region.", blobIdentifier, ns);
 
-										await _blobIndex.RemoveBlobFromRegionAsync(ns, blobIdentifier);
+										await _blobIndex.RemoveBlobFromRegionAsync(ns, blobIdentifier, cancellationToken: token);
 										deleted = true;
 									}
 								}
@@ -149,7 +150,7 @@ namespace Jupiter.Implementation
 							_logger.LogWarning("Blob {Blob} in namespace {Namespace} is of a unknown namespace, removing.", blobIdentifier, ns);
 
 							// for entries that are of a unknown namespace we simply remove them
-							await _blobIndex.RemoveBlobFromRegionAsync(ns, blobIdentifier);
+							await _blobIndex.RemoveBlobFromRegionAsync(ns, blobIdentifier, cancellationToken: cancellationToken);
 							deleted = true;
 						}
 					}
diff --git a/Engine/Source/Programs/UnrealCloudDDC/Jupiter/Implementation/Blob/BlobService.cs b/Engine/Source/Programs/UnrealCloudDDC/Jupiter/Implementation/Blob/BlobService.cs
--- a/Engine/Source/Programs/UnrealCloudDDC/Jupiter/Implementation/Blob/BlobService.cs
+++ b/Engine/Source/Programs/UnrealCloudDDC/Jupiter/Implementation/Blob/BlobService.cs
@@ -9,6 +9,7 @@ using System.Linq;
 using System.Net;
 using System.Net.Http;
 using System.Security.Claims;
+using System.Threading;
 using System.Threading.Tasks;
 using EpicGames.AspNet;
 using EpicGames.Horde.Storage;
@@ -95,12 +96,12 @@ public class BlobService : IBlobService
 		return store;
 	}
 
-	public async Task<ContentHash> VerifyContentMatchesHashAsync(Stream content, ContentHash identifier)
+	public async Task<ContentHash> VerifyContentMatchesHashAsync(Stream content, ContentHash identifier, CancellationToken cancellationToken = default)
 	{
 		ContentHash blobHash;
 		{
 			using TelemetrySpan _ = _tracer.StartActiveSpan("web.hash").SetAttribute("operation.name", "web.hash");
-			blobHash = await BlobId.FromStreamAsync(content);
+			blobHash = await BlobId.FromStreamAsync(content, cancellationToken);
 		}
 
 		if (!identifier.Equals(blobHash))
@@ -111,7 +112,7 @@ public class BlobService : IBlobService
 		return identifier;
 	}
 
-	public async Task<BlobId> PutObjectAsync(NamespaceId ns, IBufferedPayload payload, BlobId identifier)
+	public async Task<BlobId> PutObjectAsync(NamespaceId ns, IBufferedPayload payload, BlobId identifier, CancellationToken cancellationToken = default)
 	{
 		bool useContentAddressedStorage = _namespacePolicyResolver.GetPoliciesForNs(ns).UseContentAddressedStorage;
 		using TelemetrySpan scope = _tracer.StartActiveSpan("put_blob")
@@ -120,16 +121,16 @@ public class BlobService : IBlobService
 			.SetAttribute("Content-Length", payload.Length.ToString());
 
 		await using Stream hashStream = payload.GetStream();
-		BlobId id = useContentAddressedStorage ? BlobId.FromContentHash(await VerifyContentMatchesHashAsync(hashStream, identifier)) : identifier;
+		BlobId id = useContentAddressedStorage ? BlobId.FromContentHash(await VerifyContentMatchesHashAsync(hashStream, identifier, cancellationToken)) : identifier;
 
-		BlobId objectStoreIdentifier = await PutObjectToStoresAsync(ns, payload, id);
-		await _blobIndex.AddBlobToIndexAsync(ns, id);
+		BlobId objectStoreIdentifier = await PutObjectToStoresAsync(ns, payload, id, cancellationToken);
+		await _blobIndex.AddBlobToIndexAsync(ns, id, cancellationToken: cancellationToken);
 
 		return objectStoreIdentifier;
 
 	}
 
-	public async Task<BlobId> PutObjectAsync(NamespaceId ns, byte[] payload, BlobId identifier)
+	public async Task<BlobId> PutObjectAsync(NamespaceId ns, byte[] payload, BlobId identifier, CancellationToken cancellationToken = default)
 	{
 		bool useContentAddressedStorage = _namespacePolicyResolver.GetPoliciesForNs(ns).UseContentAddressedStorage;
 		using TelemetrySpan scope = _tracer.StartActiveSpan("put_blob")
@@ -139,15 +140,15 @@ public class BlobService : IBlobService
 			;
 
 		await using Stream hashStream = new MemoryStream(payload);
-		BlobId id = useContentAddressedStorage ? BlobId.FromContentHash(await VerifyContentMatchesHashAsync(hashStream, identifier)) : identifier;
+		BlobId id = useContentAddressedStorage ? BlobId.FromContentHash(await VerifyContentMatchesHashAsync(hashStream, identifier, cancellationToken)) : identifier;
 
-		BlobId objectStoreIdentifier = await PutObjectToStoresAsync(ns, payload, id);
-		await _blobIndex.AddBlobToIndexAsync(ns, id);
+		BlobId objectStoreIdentifier = await PutObjectToStoresAsync(ns, payload, id, cancellationToken);
+		await _blobIndex.AddBlobToIndexAsync(ns, id, cancellationToken: cancellationToken);
 
 		return objectStoreIdentifier;
 	}
 
-	public async Task<Uri?> MaybePutObjectWithRedirectAsync(NamespaceId ns, BlobId identifier)
+	public async Task<Uri?> MaybePutObjectWithRedirectAsync(NamespaceId ns, BlobId identifier, CancellationToken cancellationToken = default)
 	{
 		bool allowRedirectUris = _namespacePolicyResolver.GetPoliciesForNs(ns).AllowRedirectUris;
 		if (!allowRedirectUris)
@@ -181,7 +182,7 @@ public class BlobService : IBlobService
 		return null;
 	}
 
-	public async Task<BlobId> PutObjectKnownHashAsync(NamespaceId ns, IBufferedPayload content, BlobId identifier)
+	public async Task<BlobId> PutObjectKnownHashAsync(NamespaceId ns, IBufferedPayload content, BlobId identifier, CancellationToken cancellationToken = default)
 	{
 		using TelemetrySpan scope = _tracer.StartActiveSpan("put_blob")
 			.SetAttribute("operation.name", "put_blob")
@@ -189,13 +190,13 @@ public class BlobService : IBlobService
 			.SetAttribute("Content-Length", content.Length.ToString())
 			;
 
-		BlobId objectStoreIdentifier = await PutObjectToStoresAsync(ns, content, identifier);
-		await _blobIndex.AddBlobToIndexAsync(ns, identifier);
+		BlobId objectStoreIdentifier = await PutObjectToStoresAsync(ns, content, identifier, cancellationToken);
+		await _blobIndex.AddBlobToIndexAsync(ns, identifier, cancellationToken: cancellationToken);
 
 		return objectStoreIdentifier;
 	}
 
-	private async Task<BlobId> PutObjectToStoresAsync(NamespaceId ns, IBufferedPayload bufferedPayload, BlobId identifier)
+	private async Task<BlobId> PutObjectToStoresAsync(NamespaceId ns, IBufferedPayload bufferedPayload, BlobId identifier, CancellationToken cancellationToken = default)
 	{
 		IServerTiming? serverTiming = _httpContextAccessor.HttpContext?.RequestServices.GetService<IServerTiming>();
 
@@ -216,13 +217,13 @@ public class BlobService : IBlobService
 		NamespacePolicy policy = _namespacePolicyResolver.GetPoliciesForNs(ns);
 		if (policy.PopulateFallbackNamespaceOnUpload && policy.FallbackNamespace.HasValue)
 		{
-			await PutObjectToStoresAsync(policy.FallbackNamespace.Value, bufferedPayload, identifier);
-			await _blobIndex.AddBlobToIndexAsync(policy.FallbackNamespace.Value, identifier);
+			await PutObjectToStoresAsync(policy.FallbackNamespace.Value, bufferedPayload, identifier, cancellationToken);
+			await _blobIndex.AddBlobToIndexAsync(policy.FallbackNamespace.Value, identifier, cancellationToken: cancellationToken);
 		}
 		return identifier;
 	}
 
-	private async Task<BlobId> PutObjectToStoresAsync(NamespaceId ns, byte[] payload, BlobId identifier)
+	private async Task<BlobId> PutObjectToStoresAsync(NamespaceId ns, byte[] payload, BlobId identifier, CancellationToken cancellationToken = default)
 	{
 		foreach (IBlobStore store in _blobStores)
 		{
@@ -232,17 +233,17 @@ public class BlobService : IBlobService
 		NamespacePolicy policy = _namespacePolicyResolver.GetPoliciesForNs(ns);
 		if (policy.PopulateFallbackNamespaceOnUpload && policy.FallbackNamespace.HasValue)
 		{
-			await PutObjectToStoresAsync(policy.FallbackNamespace.Value, payload, identifier);
-			await _blobIndex.AddBlobToIndexAsync(policy.FallbackNamespace.Value, identifier);
+			await PutObjectToStoresAsync(policy.FallbackNamespace.Value, payload, identifier, cancellationToken);
+			await _blobIndex.AddBlobToIndexAsync(policy.FallbackNamespace.Value, identifier, cancellationToken: cancellationToken);
 		}
 		return identifier;
 	}
 
-	public async Task<BlobContents> GetObjectAsync(NamespaceId ns, BlobId blob, List<string>? storageLayers = null, bool supportsRedirectUri = false, bool allowOndemandReplication = true)
+	public async Task<BlobContents> GetObjectAsync(NamespaceId ns, BlobId blob, List<string>? storageLayers = null, bool supportsRedirectUri = false, bool allowOndemandReplication = true, CancellationToken cancellationToken = default)
 	{
 		try
 		{
-			return await GetObjectFromStoresAsync(ns, blob, storageLayers, supportsRedirectUri);
+			return await GetObjectFromStoresAsync(ns, blob, storageLayers, supportsRedirectUri, cancellationToken);
 		}
 		catch (BlobNotFoundException)
 		{
@@ -277,18 +278,18 @@ public class BlobService : IBlobService
 					IServerTiming? serverTiming = _httpContextAccessor.HttpContext?.RequestServices.GetService<IServerTiming>();
 
 					// read the content from the fallback namespace
-					BlobContents fallbackContent = await GetObjectFromStoresAsync(policy.FallbackNamespace.Value, blob, storageLayers);
+					BlobContents fallbackContent = await GetObjectFromStoresAsync(policy.FallbackNamespace.Value, blob, storageLayers, cancellationToken: cancellationToken);
 
 					// populate the primary namespace with the content
 					using TelemetrySpan _ = _tracer.StartActiveSpan("HierarchicalStore.Populate").SetAttribute("operation.name", "HierarchicalStore.Populate");
 					using ServerTimingMetricScoped? serverTimingScope = serverTiming?.CreateServerTimingMetricScope($"blob.populate", "Populating caches with blob contents");
 
 					await using MemoryStream tempStream = new MemoryStream();
-					await fallbackContent.Stream.CopyToAsync(tempStream);
+					await fallbackContent.Stream.CopyToAsync(tempStream, cancellationToken);
 					byte[] data = tempStream.ToArray();
 
-					await PutObjectAsync(ns, data, blob);
-					return await GetObjectAsync(ns, blob);
+					await PutObjectAsync(ns, data, blob, cancellationToken);
+					return await GetObjectAsync(ns, blob, cancellationToken: cancellationToken);
 				}
 				catch (BlobNotFoundException)
 				{
@@ -300,7 +301,7 @@ public class BlobService : IBlobService
 			{
 				try
 				{
-					return await ReplicateObjectAsync(ns, blob);
+					return await ReplicateObjectAsync(ns, blob, cancellationToken: cancellationToken);
 				}
 				catch (BlobNotFoundException)
 				{
@@ -313,7 +314,7 @@ public class BlobService : IBlobService
 			{
 				try
 				{
-					return await ReplicateObjectAsync(policy.FallbackNamespace.Value, blob);
+					return await ReplicateObjectAsync(policy.FallbackNamespace.Value, blob, cancellationToken: cancellationToken);
 				}
 				catch (BlobNotFoundException)
 				{
@@ -326,7 +327,7 @@ public class BlobService : IBlobService
 		}
 	}
 
-	private async Task<BlobContents> GetObjectFromStoresAsync(NamespaceId ns, BlobId blob, List<string>? storageLayers = null, bool supportsRedirectUri = false)
+	private async Task<BlobContents> GetObjectFromStoresAsync(NamespaceId ns, BlobId blob, List<string>? storageLayers = null, bool supportsRedirectUri = false, CancellationToken cancellationToken = default)
 	{
 		bool seenBlobNotFound = false;
 		bool seenNamespaceNotFound = false;
@@ -414,7 +415,7 @@ public class BlobService : IBlobService
 			using ServerTimingMetricScoped? serverTimingScope = serverTiming?.CreateServerTimingMetricScope($"blob.populate", "Populating caches with blob contents");
 
 			// not using using as the blob contents will take ownership of this buffered payload and dispose it when the contents is disposed
-			IBufferedPayload bufferedPayload = await _bufferedPayloadFactory.CreateFromStreamAsync(blobContents.Stream, blobContents.Length);
+			IBufferedPayload bufferedPayload = await _bufferedPayloadFactory.CreateFromStreamAsync(blobContents.Stream, blobContents.Length, cancellationToken);
 
 			// Don't populate the last store, as that is where we got the hit
 			for (int i = 0; i < numStoreMisses; i++)
@@ -436,7 +437,7 @@ public class BlobService : IBlobService
 		return blobContents;
 	}
 
-	public async Task<Uri?> GetObjectWithRedirectAsync(NamespaceId ns, BlobId blob, List<string>? storageLayers = null)
+	public async Task<Uri?> GetObjectWithRedirectAsync(NamespaceId ns, BlobId blob, List<string>? storageLayers = null, CancellationToken cancellationToken = default)
 	{
 		bool seenBlobNotFound = false;
 		bool seenNamespaceNotFound = false;
@@ -513,7 +514,7 @@ public class BlobService : IBlobService
 		return redirectUri;
 	}
 
-	public async Task<BlobMetadata> GetObjectMetadataAsync(NamespaceId ns, BlobId blobId)
+	public async Task<BlobMetadata> GetObjectMetadataAsync(NamespaceId ns, BlobId blobId, CancellationToken cancellationToken)
 	{
 		bool seenBlobNotFound = false;
 		bool seenNamespaceNotFound = false;
@@ -562,7 +563,7 @@ public class BlobService : IBlobService
 		throw new BlobNotFoundException(ns, blobId);
 	}
 
-	public async Task<BlobContents> ReplicateObjectAsync(NamespaceId ns, BlobId blob, bool force = false)
+	public async Task<BlobContents> ReplicateObjectAsync(NamespaceId ns, BlobId blob, bool force = false, CancellationToken cancellationToken = default)
 	{
 		if (!force && !ShouldFetchBlobOnDemand(ns))
 		{
@@ -573,7 +574,7 @@ public class BlobService : IBlobService
 
 		using ServerTimingMetricScoped? serverTimingScope = serverTiming?.CreateServerTimingMetricScope("blob.replicate", "Replicating blob from remote instances");
 
-		List<string> regions = await _blobIndex.GetBlobRegionsAsync(ns, blob);
+		List<string> regions = await _blobIndex.GetBlobRegionsAsync(ns, blob, cancellationToken);
 
 		if (!regions.Any())
 		{
@@ -602,7 +603,7 @@ public class BlobService : IBlobService
 				url += "/";
 			}
 			using HttpRequestMessage blobRequest = await BuildHttpRequestAsync(HttpMethod.Get, new Uri($"{url}api/v1/blobs/{ns}/{blob}?allowOndemandReplication=false"));
-			HttpResponseMessage blobResponse = await httpClient.SendAsync(blobRequest, HttpCompletionOption.ResponseHeadersRead);
+			HttpResponseMessage blobResponse = await httpClient.SendAsync(blobRequest, HttpCompletionOption.ResponseHeadersRead, cancellationToken);
 
 			if (blobResponse.StatusCode == HttpStatusCode.NotFound)
 			{
@@ -615,9 +616,9 @@ public class BlobService : IBlobService
 				throw new BlobReplicationException(ns, blob, $"Failed to replicate {blob} in {ns} from region {region} due to bad http status code {blobResponse.StatusCode}.");
 			}
 
-			await using Stream s = await blobResponse.Content.ReadAsStreamAsync();
-			using IBufferedPayload payload = await _bufferedPayloadFactory.CreateFilesystemBufferedPayloadAsync(s);
-			await PutObjectAsync(ns, payload, blob);
+			await using Stream s = await blobResponse.Content.ReadAsStreamAsync(cancellationToken);
+			using IBufferedPayload payload = await _bufferedPayloadFactory.CreateFilesystemBufferedPayloadAsync(s, cancellationToken);
+			await PutObjectAsync(ns, payload, blob, cancellationToken);
 			replicated = true;
 			break;
 		}
@@ -627,7 +628,7 @@ public class BlobService : IBlobService
 			throw new BlobReplicationException(ns, blob, $"Failed to replicate {blob} in {ns} due to it not existing in any region");
 		}
 
-		return await GetObjectAsync(ns, blob);
+		return await GetObjectAsync(ns, blob, cancellationToken: cancellationToken);
 	}
 
 	public bool ShouldFetchBlobOnDemand(NamespaceId ns)
@@ -647,9 +648,9 @@ public class BlobService : IBlobService
 		return request;
 	}
 
-	public async Task<bool> ExistsAsync(NamespaceId ns, BlobId blob, List<string>? storageLayers = null)
+	public async Task<bool> ExistsAsync(NamespaceId ns, BlobId blob, List<string>? storageLayers = null, CancellationToken cancellationToken = default)
 	{
-		bool exists = await ExistsInStoresAsync(ns, blob, storageLayers);
+		bool exists = await ExistsInStoresAsync(ns, blob, storageLayers, cancellationToken);
 		if (exists)
 		{
 			return exists;
@@ -658,18 +659,18 @@ public class BlobService : IBlobService
 		NamespacePolicy policy = _namespacePolicyResolver.GetPoliciesForNs(ns);
 		if (policy.FallbackNamespace != null)
 		{
-			return await ExistsInStoresAsync(policy.FallbackNamespace.Value, blob, storageLayers);
+			return await ExistsInStoresAsync(policy.FallbackNamespace.Value, blob, storageLayers, cancellationToken);
 		}
 
 		if (ShouldFetchBlobOnDemand(ns))
 		{
-			return await ExistsInRemoteAsync(ns, blob);
+			return await ExistsInRemoteAsync(ns, blob, cancellationToken);
 		}
 
 		return false;
 	}
 
-	private async Task<bool> ExistsInStoresAsync(NamespaceId ns, BlobId blob, List<string>? storageLayers = null)
+	private async Task<bool> ExistsInStoresAsync(NamespaceId ns, BlobId blob, List<string>? storageLayers = null, CancellationToken cancellationToken = default)
 	{
 		bool useBlobIndex = _namespacePolicyResolver.GetPoliciesForNs(ns).UseBlobIndexForExists;
 		if (useBlobIndex)
@@ -679,7 +680,7 @@ public class BlobService : IBlobService
 				.SetAttribute("resource.name", blob.ToString())
 				.SetAttribute("BlobStore", "BlobIndex")
 				;
-			bool exists = await _blobIndex.BlobExistsInRegionAsync(ns, blob);
+			bool exists = await _blobIndex.BlobExistsInRegionAsync(ns, blob, cancellationToken: cancellationToken);
 			if (exists)
 			{
 				scope.SetAttribute("ObjectFound", true.ToString());
@@ -730,14 +731,14 @@ public class BlobService : IBlobService
 		}
 	}
 
-	public async Task<bool> ExistsInRemoteAsync(NamespaceId ns, BlobId blob)
+	public async Task<bool> ExistsInRemoteAsync(NamespaceId ns, BlobId blob, CancellationToken cancellationToken)
 	{
 		IServerTiming? serverTiming = _httpContextAccessor.HttpContext?.RequestServices.GetService<IServerTiming>();
 		using TelemetrySpan scope = _tracer.StartActiveSpan("HierarchicalStore.ExistsRemote").SetAttribute("operation.name", "HierarchicalStore.ExistsRemote");
 
 		using ServerTimingMetricScoped? serverTimingScope = serverTiming?.CreateServerTimingMetricScope("blob.exists-remote", "Verify if blob exists in remotes");
 
-		List<string> regions = await _blobIndex.GetBlobRegionsAsync(ns, blob);
+		List<string> regions = await _blobIndex.GetBlobRegionsAsync(ns, blob, cancellationToken);
 
 		// we do not actually verify that the blob exists remotely as that would take a lot of time
 		// instead we simply check if there are any regions were the blob exists that is not our current region
@@ -756,7 +757,7 @@ public class BlobService : IBlobService
 		return false;
 	}
 
-	public async Task<bool> ExistsInRootStoreAsync(NamespaceId ns, BlobId blob)
+	public async Task<bool> ExistsInRootStoreAsync(NamespaceId ns, BlobId blob, CancellationToken cancellationToken = default)
 	{
 		IBlobStore store = _blobStores.Last();
 
@@ -774,7 +775,7 @@ public class BlobService : IBlobService
 		return false;
 	}
 
-	public async Task DeleteObjectAsync(NamespaceId ns, BlobId blob)
+	public async Task DeleteObjectAsync(NamespaceId ns, BlobId blob, CancellationToken cancellationToken = default)
 	{
 		bool blobNotFound = false;
 		bool deletedAtLeastOnce = false;
@@ -782,8 +783,8 @@ public class BlobService : IBlobService
 		// remove the object from the tracking first, if this times out we do not want to end up with a inconsistent blob index
 		// if the blob store delete fails on the other hand we will still run a delete again during GC (as the blob is still orphaned at that point)
 		// this assumes that blob gc is based on scanning the root blob store
-		await _blobIndex.RemoveBlobFromRegionAsync(ns, blob);
-		await _blobIndex.RemoveReferencesAsync(ns, blob, null);
+		await _blobIndex.RemoveBlobFromRegionAsync(ns, blob, cancellationToken: cancellationToken);
+		await _blobIndex.RemoveReferencesAsync(ns, blob, null, cancellationToken);
 
 		foreach (IBlobStore store in _blobStores)
 		{
@@ -821,7 +822,7 @@ public class BlobService : IBlobService
 		throw new NamespaceNotFoundException(ns);
 	}
 
-	public async Task DeleteNamespaceAsync(NamespaceId ns)
+	public async Task DeleteNamespaceAsync(NamespaceId ns, CancellationToken cancellationToken = default)
 	{
 		bool deletedAtLeastOnce = false;
 		foreach (IBlobStore store in _blobStores)
@@ -850,18 +851,18 @@ public class BlobService : IBlobService
 		throw new NamespaceNotFoundException(ns);
 	}
 
-	public IAsyncEnumerable<(BlobId, DateTime)> ListObjectsAsync(NamespaceId ns)
+	public IAsyncEnumerable<(BlobId, DateTime)> ListObjectsAsync(NamespaceId ns, CancellationToken cancellationToken)
 	{
 		// as this is a hierarchy of blob stores the last blob store should contain the superset of all stores
 		return _blobStores.Last().ListObjectsAsync(ns);
 	}
 
-	public async Task<BlobId[]> FilterOutKnownBlobsAsync(NamespaceId ns, IEnumerable<BlobId> blobs)
+	public async Task<BlobId[]> FilterOutKnownBlobsAsync(NamespaceId ns, IEnumerable<BlobId> blobs, CancellationToken cancellationToken)
 	{
 		List<(BlobId, Task<bool>)> existTasks = new();
 		foreach (BlobId blob in blobs)
 		{
-			existTasks.Add((blob, ExistsAsync(ns, blob)));
+			existTasks.Add((blob, ExistsAsync(ns, blob, cancellationToken: cancellationToken)));
 		}
 
 		List<BlobId> missingBlobs = new();
@@ -878,15 +879,15 @@ public class BlobService : IBlobService
 		return missingBlobs.ToArray();
 	}
 
-	public async Task<BlobId[]> FilterOutKnownBlobsAsync(NamespaceId ns, IAsyncEnumerable<BlobId> blobs)
+	public async Task<BlobId[]> FilterOutKnownBlobsAsync(NamespaceId ns, IAsyncEnumerable<BlobId> blobs, CancellationToken cancellationToken)
 	{
 		ConcurrentBag<BlobId> missingBlobs = new ConcurrentBag<BlobId>();
 
 		try
 		{
-			await Parallel.ForEachAsync(blobs, async (identifier, ctx) =>
+			await Parallel.ForEachAsync(blobs, cancellationToken, async (identifier, ctx) =>
 			{
-				bool exists = await ExistsAsync(ns, identifier);
+				bool exists = await ExistsAsync(ns, identifier, cancellationToken: ctx);
 
 				if (!exists)
 				{
@@ -907,13 +908,13 @@ public class BlobService : IBlobService
 		return missingBlobs.ToArray();
 	}
 
-	public async Task<BlobContents> GetObjectsAsync(NamespaceId ns, BlobId[] blobs)
+	public async Task<BlobContents> GetObjectsAsync(NamespaceId ns, BlobId[] blobs, CancellationToken cancellationToken)
 	{
 		using TelemetrySpan _ = _tracer.StartActiveSpan("blob.combine").SetAttribute("operation.name", "blob.combine");
 		Task<BlobContents>[] tasks = new Task<BlobContents>[blobs.Length];
 		for (int i = 0; i < blobs.Length; i++)
 		{
-			tasks[i] = GetObjectAsync(ns, blobs[i]);
+			tasks[i] = GetObjectAsync(ns, blobs[i], cancellationToken: cancellationToken);
 		}
 
 		MemoryStream ms = new MemoryStream();
@@ -921,7 +922,7 @@ public class BlobService : IBlobService
 		{
 			BlobContents blob = await task;
 			await using Stream s = blob.Stream;
-			await s.CopyToAsync(ms);
+			await s.CopyToAsync(ms, cancellationToken);
 		}
 
 		ms.Seek(0, SeekOrigin.Begin);
diff --git a/Engine/Source/Programs/UnrealCloudDDC/Jupiter/Implementation/Blob/BlobStoreConsistencyCheckService.cs b/Engine/Source/Programs/UnrealCloudDDC/Jupiter/Implementation/Blob/BlobStoreConsistencyCheckService.cs
--- a/Engine/Source/Programs/UnrealCloudDDC/Jupiter/Implementation/Blob/BlobStoreConsistencyCheckService.cs
+++ b/Engine/Source/Programs/UnrealCloudDDC/Jupiter/Implementation/Blob/BlobStoreConsistencyCheckService.cs
@@ -58,12 +58,12 @@ namespace Jupiter.Implementation
 				return false;
 			}
 
-			await RunConsistencyCheckAsync();
+			await RunConsistencyCheckAsync(cancellationToken);
 
 			return true;
 		}
 
-		private async Task RunConsistencyCheckAsync()
+		private async Task RunConsistencyCheckAsync(CancellationToken cancellationToken)
 		{
 			foreach (IBlobStore blobStore in BlobService.GetBlobStores(_provider, _unrealCloudDDCSettings).Where(RunConsistencyCheckOnBlobStore))
 			{
@@ -83,7 +83,7 @@ namespace Jupiter.Implementation
 					continue;
 				}
 
-				List<NamespaceId> namespaces = await _referencesStore.GetNamespacesAsync().ToListAsync();
+				List<NamespaceId> namespaces = await _referencesStore.GetNamespacesAsync(cancellationToken).ToListAsync(cancellationToken);
 
 				// technically this does not need to be run per namespace but per storage pool
 				foreach (NamespaceId ns in namespaces)
@@ -115,7 +115,7 @@ namespace Jupiter.Implementation
 						await using Stream s = contents.Stream;
 
 						bool inconsistencyFound = false;
-						BlobId newHash = await BlobId.FromStreamAsync(s);
+						BlobId newHash = await BlobId.FromStreamAsync(s, cancellationToken);
 						if (!blob.Equals(newHash))
 						{
 							_logger.LogError("Mismatching hash for {Blob} in {Namespace} stored in {BlobStore}, new hash has {NewHash}. Deleting incorrect blob.", blob, ns, blobStoreName, newHash);
@@ -126,7 +126,7 @@ namespace Jupiter.Implementation
 							if (isRootStore)
 							{
 								// update blob index tracking to indicate that we no longer have this blob in this region
-								await _blobIndex.RemoveBlobFromRegionAsync(ns, blob);
+								await _blobIndex.RemoveBlobFromRegionAsync(ns, blob, cancellationToken: cancellationToken);
 							}
 						}
 
diff --git a/Engine/Source/Programs/UnrealCloudDDC/Jupiter/Implementation/Blob/CachedBlobIndex.cs b/Engine/Source/Programs/UnrealCloudDDC/Jupiter/Implementation/Blob/CachedBlobIndex.cs
--- a/Engine/Source/Programs/UnrealCloudDDC/Jupiter/Implementation/Blob/CachedBlobIndex.cs
+++ b/Engine/Source/Programs/UnrealCloudDDC/Jupiter/Implementation/Blob/CachedBlobIndex.cs
@@ -2,6 +2,7 @@
 
 using System;
 using System.Collections.Generic;
+using System.Threading;
 using System.Threading.Tasks;
 using EpicGames.Horde.Storage;
 
@@ -16,66 +17,66 @@ public class CachedBlobIndex : IBlobIndex
 		_fileSystemStore = fileSystemStore;
 	}
 
-	public async Task AddBlobToIndexAsync(NamespaceId ns, BlobId id, string? region = null)
+	public async Task AddBlobToIndexAsync(NamespaceId ns, BlobId id, string? region = null, CancellationToken cancellationToken = default)
 	{
 		// We do not actually track any blob information when running in cached mode
 		await Task.CompletedTask;
 	}
 
-	public async Task RemoveBlobFromRegionAsync(NamespaceId ns, BlobId id, string? region = null)
+	public async Task RemoveBlobFromRegionAsync(NamespaceId ns, BlobId id, string? region = null, CancellationToken cancellationToken = default)
 	{
 		// We do not actually track any blob information when running in cached mode
 		await Task.CompletedTask;
 	}
 
-	public IAsyncEnumerable<BaseBlobReference> GetBlobReferencesAsync(NamespaceId ns, BlobId id)
+	public IAsyncEnumerable<BaseBlobReference> GetBlobReferencesAsync(NamespaceId ns, BlobId id, CancellationToken cancellationToken)
 	{
 		throw new NotImplementedException();
 	}
 
-	public async Task<bool> BlobExistsInRegionAsync(NamespaceId ns, BlobId blobIdentifier, string? region = null)
+	public async Task<bool> BlobExistsInRegionAsync(NamespaceId ns, BlobId blobIdentifier, string? region = null, CancellationToken cancellationToken = default)
 	{
 		return await _fileSystemStore.ExistsAsync(ns, blobIdentifier, forceCheck: false);
 	}
 
-	public async Task AddRefToBlobsAsync(NamespaceId ns, BucketId bucket, RefId key, BlobId[] blobs)
+	public async Task AddRefToBlobsAsync(NamespaceId ns, BucketId bucket, RefId key, BlobId[] blobs, CancellationToken cancellationToken)
 	{
 		// We do not actually track any blob information when running in cached mode
 		await Task.CompletedTask;
 	}
 
-	public IAsyncEnumerable<(NamespaceId, BlobId)> GetAllBlobsAsync()
+	public IAsyncEnumerable<(NamespaceId, BlobId)> GetAllBlobsAsync(CancellationToken cancellationToken)
 	{
 		throw new NotImplementedException();
 	}
 
-	public Task RemoveReferencesAsync(NamespaceId ns, BlobId id, List<BaseBlobReference>? referencesToRemove)
+	public Task RemoveReferencesAsync(NamespaceId ns, BlobId id, List<BaseBlobReference>? referencesToRemove, CancellationToken cancellationToken)
 	{
 		// We do not actually track any blob information when running in cached mode
 		return Task.CompletedTask;
 	}
 
-	public Task<List<string>> GetBlobRegionsAsync(NamespaceId ns, BlobId blob)
+	public Task<List<string>> GetBlobRegionsAsync(NamespaceId ns, BlobId blob, CancellationToken cancellationToken)
 	{
 		throw new NotImplementedException();
 	}
 
-	public Task AddBlobReferencesAsync(NamespaceId ns, BlobId sourceBlob, BlobId targetBlob)
+	public Task AddBlobReferencesAsync(NamespaceId ns, BlobId sourceBlob, BlobId targetBlob, CancellationToken cancellationToken)
 	{
 		throw new NotImplementedException();
 	}
 
-	public Task AddBlobToBucketListAsync(NamespaceId ns, BucketId bucket, RefId key, BlobId blobId, long blobSize)
+	public Task AddBlobToBucketListAsync(NamespaceId ns, BucketId bucket, RefId key, BlobId blobId, long blobSize, CancellationToken cancellationToken)
 	{
 		throw new NotImplementedException();
 	}
 
-	public Task RemoveBlobFromBucketListAsync(NamespaceId ns, BucketId bucket, RefId key, List<BlobId> blobIds)
+	public Task RemoveBlobFromBucketListAsync(NamespaceId ns, BucketId bucket, RefId key, List<BlobId> blobIds, CancellationToken cancellationToken)
 	{
 		throw new NotImplementedException();
 	}
 
-	public Task<BucketStats> CalculateBucketStatisticsAsync(NamespaceId ns, BucketId bucket)
+	public Task<BucketStats> CalculateBucketStatisticsAsync(NamespaceId ns, BucketId bucket, CancellationToken cancellationToken)
 	{
 		throw new NotImplementedException();
 	}
diff --git a/Engine/Source/Programs/UnrealCloudDDC/Jupiter/Implementation/Blob/IBlobIndex.cs b/Engine/Source/Programs/UnrealCloudDDC/Jupiter/Implementation/Blob/IBlobIndex.cs
--- a/Engine/Source/Programs/UnrealCloudDDC/Jupiter/Implementation/Blob/IBlobIndex.cs
+++ b/Engine/Source/Programs/UnrealCloudDDC/Jupiter/Implementation/Blob/IBlobIndex.cs
@@ -1,6 +1,7 @@
 // Copyright Epic Games, Inc. All Rights Reserved.
 
 using System.Collections.Generic;
+using System.Threading;
 using System.Threading.Tasks;
 using EpicGames.Horde.Storage;
 
@@ -20,23 +21,23 @@ namespace Jupiter.Implementation.Blob
 
 	public interface IBlobIndex
 	{
-		Task AddBlobToIndexAsync(NamespaceId ns, BlobId id, string? region = null);
+		Task AddBlobToIndexAsync(NamespaceId ns, BlobId id, string? region = null, CancellationToken cancellationToken = default);
 
-		Task RemoveBlobFromRegionAsync(NamespaceId ns, BlobId id, string? region = null);
+		Task RemoveBlobFromRegionAsync(NamespaceId ns, BlobId id, string? region = null, CancellationToken cancellationToken = default);
 
-		Task<bool> BlobExistsInRegionAsync(NamespaceId ns, BlobId blobIdentifier, string? region = null);
-		IAsyncEnumerable<(NamespaceId, BlobId)> GetAllBlobsAsync();
+		Task<bool> BlobExistsInRegionAsync(NamespaceId ns, BlobId blobIdentifier, string? region = null, CancellationToken cancellationToken = default);
+		IAsyncEnumerable<(NamespaceId, BlobId)> GetAllBlobsAsync(CancellationToken cancellationToken = default);
 
-		IAsyncEnumerable<BaseBlobReference> GetBlobReferencesAsync(NamespaceId ns, BlobId id);
-		Task AddRefToBlobsAsync(NamespaceId ns, BucketId bucket, RefId key, BlobId[] blobs);
+		IAsyncEnumerable<BaseBlobReference> GetBlobReferencesAsync(NamespaceId ns, BlobId id, CancellationToken cancellationToken = default);
+		Task AddRefToBlobsAsync(NamespaceId ns, BucketId bucket, RefId key, BlobId[] blobs, CancellationToken cancellationToken = default);
 
-		Task RemoveReferencesAsync(NamespaceId ns, BlobId id, List<BaseBlobReference>? referencesToRemove);
-		Task<List<string>> GetBlobRegionsAsync(NamespaceId ns, BlobId blob);
-		Task AddBlobReferencesAsync(NamespaceId ns, BlobId sourceBlob, BlobId targetBlob);
+		Task RemoveReferencesAsync(NamespaceId ns, BlobId id, List<BaseBlobReference>? referencesToRemove, CancellationToken cancellationToken = default);
+		Task<List<string>> GetBlobRegionsAsync(NamespaceId ns, BlobId blob, CancellationToken cancellationToken = default);
+		Task AddBlobReferencesAsync(NamespaceId ns, BlobId sourceBlob, BlobId targetBlob, CancellationToken cancellationToken = default);
 
-		Task AddBlobToBucketListAsync(NamespaceId ns, BucketId bucket, RefId key, BlobId blobId, long blobSize);
-		Task RemoveBlobFromBucketListAsync(NamespaceId ns, BucketId bucket, RefId key, List<BlobId> blobIds);
-		Task<BucketStats> CalculateBucketStatisticsAsync(NamespaceId ns, BucketId bucket);
+		Task AddBlobToBucketListAsync(NamespaceId ns, BucketId bucket, RefId key, BlobId blobId, long blobSize, CancellationToken cancellationToken = default);
+		Task RemoveBlobFromBucketListAsync(NamespaceId ns, BucketId bucket, RefId key, List<BlobId> blobIds, CancellationToken cancellationToken = default);
+		Task<BucketStats> CalculateBucketStatisticsAsync(NamespaceId ns, BucketId bucket, CancellationToken cancellationToken = default);
 	}
 
 	public abstract class BaseBlobReference
diff --git a/Engine/Source/Programs/UnrealCloudDDC/Jupiter/Implementation/Blob/IBlobService.cs b/Engine/Source/Programs/UnrealCloudDDC/Jupiter/Implementation/Blob/IBlobService.cs
--- a/Engine/Source/Programs/UnrealCloudDDC/Jupiter/Implementation/Blob/IBlobService.cs
+++ b/Engine/Source/Programs/UnrealCloudDDC/Jupiter/Implementation/Blob/IBlobService.cs
@@ -4,6 +4,7 @@ using System;
 using System.Collections.Generic;
 using System.IO;
 using System.Net.Mime;
+using System.Threading;
 using System.Threading.Tasks;
 using EpicGames.AspNet;
 using EpicGames.Horde.Storage;
@@ -15,40 +16,41 @@ namespace Jupiter.Implementation;
 
 public interface IBlobService
 {
-	Task<ContentHash> VerifyContentMatchesHashAsync(Stream content, ContentHash identifier);
-	Task<BlobId> PutObjectKnownHashAsync(NamespaceId ns, IBufferedPayload content, BlobId identifier);
-	Task<BlobId> PutObjectAsync(NamespaceId ns, IBufferedPayload payload, BlobId identifier);
-	Task<BlobId> PutObjectAsync(NamespaceId ns, byte[] payload, BlobId identifier);
-	Task<Uri?> MaybePutObjectWithRedirectAsync(NamespaceId ns, BlobId identifier);
+	Task<ContentHash> VerifyContentMatchesHashAsync(Stream content, ContentHash identifier, CancellationToken cancellationToken = default);
+	Task<BlobId> PutObjectKnownHashAsync(NamespaceId ns, IBufferedPayload content, BlobId identifier, CancellationToken cancellationToken = default);
+	Task<BlobId> PutObjectAsync(NamespaceId ns, IBufferedPayload payload, BlobId identifier, CancellationToken cancellationToken = default);
+	Task<BlobId> PutObjectAsync(NamespaceId ns, byte[] payload, BlobId identifier, CancellationToken cancellationToken = default);
+	Task<Uri?> MaybePutObjectWithRedirectAsync(NamespaceId ns, BlobId identifier, CancellationToken cancellationToken = default);
 
-	Task<BlobContents> GetObjectAsync(NamespaceId ns, BlobId blob, List<string>? storageLayers = null, bool supportsRedirectUri = false, bool allowOndemandReplication = true);
+	Task<BlobContents> GetObjectAsync(NamespaceId ns, BlobId blob, List<string>? storageLayers = null, bool supportsRedirectUri = false, bool allowOndemandReplication = true, CancellationToken cancellationToken = default);
 
-	Task<Uri?> GetObjectWithRedirectAsync(NamespaceId ns, BlobId blobIdentifier, List<string>? storageLayers = null);
+	Task<Uri?> GetObjectWithRedirectAsync(NamespaceId ns, BlobId blobIdentifier, List<string>? storageLayers = null, CancellationToken cancellationToken = default);
 
-	Task<BlobMetadata> GetObjectMetadataAsync(NamespaceId ns, BlobId blobId);
+	Task<BlobMetadata> GetObjectMetadataAsync(NamespaceId ns, BlobId blobId, CancellationToken cancellationToken = default);
 
-	Task<BlobContents> ReplicateObjectAsync(NamespaceId ns, BlobId blob, bool force = false);
+	Task<BlobContents> ReplicateObjectAsync(NamespaceId ns, BlobId blob, bool force = false, CancellationToken cancellationToken = default);
 
-	Task<bool> ExistsAsync(NamespaceId ns, BlobId blob, List<string>? storageLayers = null);
+	Task<bool> ExistsAsync(NamespaceId ns, BlobId blob, List<string>? storageLayers = null, CancellationToken cancellationToken = default);
 
 	/// <summary>
 	/// Checks that the blob exists in the root store, the store which is last in the list and thus is intended to have every blob in it
 	/// </summary>
 	/// <param name="ns">The namespace</param>
 	/// <param name="blob">The identifier of the blob</param>
+	/// <param name="cancellationToken"></param>
 	/// <returns></returns>
-	Task<bool> ExistsInRootStoreAsync(NamespaceId ns, BlobId blob);
+	Task<bool> ExistsInRootStoreAsync(NamespaceId ns, BlobId blob, CancellationToken cancellationToken = default);
 
 	// Delete a object
-	Task DeleteObjectAsync(NamespaceId ns, BlobId blob);
+	Task DeleteObjectAsync(NamespaceId ns, BlobId blob, CancellationToken cancellationToken = default);
 
 	// delete the whole namespace
-	Task DeleteNamespaceAsync(NamespaceId ns);
+	Task DeleteNamespaceAsync(NamespaceId ns, CancellationToken cancellationToken = default);
 
-	IAsyncEnumerable<(BlobId, DateTime)> ListObjectsAsync(NamespaceId ns);
-	Task<BlobId[]> FilterOutKnownBlobsAsync(NamespaceId ns, IEnumerable<BlobId> blobs);
-	Task<BlobId[]> FilterOutKnownBlobsAsync(NamespaceId ns, IAsyncEnumerable<BlobId> blobs);
-	Task<BlobContents> GetObjectsAsync(NamespaceId ns, BlobId[] refRequestBlobReferences);
+	IAsyncEnumerable<(BlobId, DateTime)> ListObjectsAsync(NamespaceId ns, CancellationToken cancellationToken = default);
+	Task<BlobId[]> FilterOutKnownBlobsAsync(NamespaceId ns, IEnumerable<BlobId> blobs, CancellationToken cancellationToken = default);
+	Task<BlobId[]> FilterOutKnownBlobsAsync(NamespaceId ns, IAsyncEnumerable<BlobId> blobs, CancellationToken cancellationToken = default);
+	Task<BlobContents> GetObjectsAsync(NamespaceId ns, BlobId[] refRequestBlobReferences, CancellationToken cancellationToken = default);
 
 	bool ShouldFetchBlobOnDemand(NamespaceId ns);
 }
@@ -67,7 +69,7 @@ public class BlobMetadata
 
 public static class BlobServiceExtensions
 {
-	public static async Task<ContentId> PutCompressedObjectAsync(this IBlobService blobService, NamespaceId ns, IBufferedPayload payload, ContentId? id, IServiceProvider provider)
+	public static async Task<ContentId> PutCompressedObjectAsync(this IBlobService blobService, NamespaceId ns, IBufferedPayload payload, ContentId? id, IServiceProvider provider, CancellationToken cancellationToken)
 	{
 		IContentIdStore contentIdStore = provider.GetService<IContentIdStore>()!;
 		CompressedBufferUtils compressedBufferUtils = provider.GetService<CompressedBufferUtils>()!;
@@ -76,20 +78,20 @@ public static class BlobServiceExtensions
 		// decompress the content and generate a identifier from it to verify the identifier we got
 		await using Stream decompressStream = payload.GetStream();
 
-		using IBufferedPayload bufferedPayload = await compressedBufferUtils.DecompressContentAsync(decompressStream, (ulong)payload.Length);
+		using IBufferedPayload bufferedPayload = await compressedBufferUtils.DecompressContentAsync(decompressStream, (ulong)payload.Length, cancellationToken);
 		await using Stream decompressedStream = bufferedPayload.GetStream();
 
 		ContentId identifierDecompressedPayload;
 		if (id != null)
 		{
-			identifierDecompressedPayload = ContentId.FromContentHash(await blobService.VerifyContentMatchesHashAsync(decompressedStream, id));
+			identifierDecompressedPayload = ContentId.FromContentHash(await blobService.VerifyContentMatchesHashAsync(decompressedStream, id, cancellationToken));
 		}
 		else
 		{
 			ContentHash blobHash;
 			{
 				using TelemetrySpan _ = tracer.StartActiveSpan("web.hash").SetAttribute("operation.name", "web.hash");
-				blobHash = await BlobId.FromStreamAsync(decompressedStream);
+				blobHash = await BlobId.FromStreamAsync(decompressedStream, cancellationToken);
 			}
 
 			identifierDecompressedPayload = ContentId.FromContentHash(blobHash);
@@ -99,17 +101,17 @@ public static class BlobServiceExtensions
 		{
 			using TelemetrySpan _ = tracer.StartActiveSpan("web.hash").SetAttribute("operation.name", "web.hash");
 			await using Stream hashStream = payload.GetStream();
-			identifierCompressedPayload = await BlobId.FromStreamAsync(hashStream);
+			identifierCompressedPayload = await BlobId.FromStreamAsync(hashStream, cancellationToken);
 		}
 
 		// commit the mapping from the decompressed hash to the compressed hash, we run this in parallel with the blob store submit
 		// TODO: let users specify weight of the blob compared to previously submitted content ids
 		int contentIdWeight = (int)payload.Length;
-		Task contentIdStoreTask = contentIdStore.PutAsync(ns, identifierDecompressedPayload, identifierCompressedPayload, contentIdWeight);
+		Task contentIdStoreTask = contentIdStore.PutAsync(ns, identifierDecompressedPayload, identifierCompressedPayload, contentIdWeight, cancellationToken);
 
 		// we still commit the compressed buffer to the object store using the hash of the compressed content
 		{
-			await blobService.PutObjectKnownHashAsync(ns, payload, identifierCompressedPayload);
+			await blobService.PutObjectKnownHashAsync(ns, payload, identifierCompressedPayload, cancellationToken);
 		}
 
 		await contentIdStoreTask;
@@ -117,12 +119,12 @@ public static class BlobServiceExtensions
 		return identifierDecompressedPayload;
 	}
 
-	public static async Task<(BlobContents, string)> GetCompressedObjectAsync(this IBlobService blobService, NamespaceId ns, ContentId contentId, IServiceProvider provider, bool supportsRedirectUri = false)
+	public static async Task<(BlobContents, string)> GetCompressedObjectAsync(this IBlobService blobService, NamespaceId ns, ContentId contentId, IServiceProvider provider, bool supportsRedirectUri = false, CancellationToken cancellationToken = default)
 	{
 		IContentIdStore contentIdStore = provider.GetService<IContentIdStore>()!;
 		Tracer tracer = provider.GetService<Tracer>()!;
 
-		BlobId[]? chunks = await contentIdStore.ResolveAsync(ns, contentId, mustBeContentId: false);
+		BlobId[]? chunks = await contentIdStore.ResolveAsync(ns, contentId, mustBeContentId: false, cancellationToken);
 		if (chunks == null || chunks.Length == 0)
 		{
 			throw new ContentIdResolveException(contentId);
@@ -139,7 +141,7 @@ public static class BlobServiceExtensions
 				mimeType = MediaTypeNames.Application.Octet;
 			}
 
-			return (await blobService.GetObjectAsync(ns, blobToReturn, supportsRedirectUri: supportsRedirectUri), mimeType);
+			return (await blobService.GetObjectAsync(ns, blobToReturn, supportsRedirectUri: supportsRedirectUri, cancellationToken: cancellationToken), mimeType);
 		}
 
 		// chunked content, combine the chunks into a single stream
@@ -148,7 +150,7 @@ public static class BlobServiceExtensions
 		for (int i = 0; i < chunks.Length; i++)
 		{
 			// even if it was requested to support redirect, since we need to combine the chunks using redirects is not possible
-			tasks[i] = blobService.GetObjectAsync(ns, chunks[i], supportsRedirectUri: false);
+			tasks[i] = blobService.GetObjectAsync(ns, chunks[i], supportsRedirectUri: false, cancellationToken: cancellationToken);
 		}
 
 		MemoryStream ms = new MemoryStream();
@@ -156,7 +158,7 @@ public static class BlobServiceExtensions
 		{
 			BlobContents blob = await task;
 			await using Stream s = blob.Stream;
-			await s.CopyToAsync(ms);
+			await s.CopyToAsync(ms, cancellationToken);
 		}
 
 		ms.Seek(0, SeekOrigin.Begin);
diff --git a/Engine/Source/Programs/UnrealCloudDDC/Jupiter/Implementation/Blob/MemoryBlobIndex.cs b/Engine/Source/Programs/UnrealCloudDDC/Jupiter/Implementation/Blob/MemoryBlobIndex.cs
--- a/Engine/Source/Programs/UnrealCloudDDC/Jupiter/Implementation/Blob/MemoryBlobIndex.cs
+++ b/Engine/Source/Programs/UnrealCloudDDC/Jupiter/Implementation/Blob/MemoryBlobIndex.cs
@@ -4,6 +4,8 @@ using System;
 using System.Collections.Concurrent;
 using System.Collections.Generic;
 using System.Linq;
+using System.Runtime.CompilerServices;
+using System.Threading;
 using System.Threading.Tasks;
 using EpicGames.Horde.Storage;
 using Microsoft.Extensions.Options;
@@ -59,7 +61,7 @@ public class MemoryBlobIndex : IBlobIndex
 		return _index.GetOrAdd(ns, id => new ConcurrentDictionary<BlobId, MemoryBlobInfo>());
 	}
 
-	public Task AddBlobToIndexAsync(NamespaceId ns, BlobId id, string? region = null)
+	public Task AddBlobToIndexAsync(NamespaceId ns, BlobId id, string? region = null, CancellationToken cancellationToken = default)
 	{
 		region ??= _jupiterSettings.CurrentValue.CurrentSite;
 		ConcurrentDictionary<BlobId, MemoryBlobInfo> index = GetNamespaceContainer(ns);
@@ -79,7 +81,7 @@ public class MemoryBlobIndex : IBlobIndex
 		return Task.FromResult<MemoryBlobInfo?>(blobInfo);
 	}
 
-	public Task RemoveBlobFromRegionAsync(NamespaceId ns, BlobId id, string? region = null)
+	public Task RemoveBlobFromRegionAsync(NamespaceId ns, BlobId id, string? region = null, CancellationToken cancellationToken = default)
 	{
 		region ??= _jupiterSettings.CurrentValue.CurrentSite;
 		ConcurrentDictionary<BlobId, MemoryBlobInfo> index = GetNamespaceContainer(ns);
@@ -96,14 +98,14 @@ public class MemoryBlobIndex : IBlobIndex
 		return Task.CompletedTask;
 	}
 
-	public async Task<bool> BlobExistsInRegionAsync(NamespaceId ns, BlobId blobIdentifier, string? region = null)
+	public async Task<bool> BlobExistsInRegionAsync(NamespaceId ns, BlobId blobIdentifier, string? region = null, CancellationToken cancellationToken = default)
 	{
 		string expectedRegion = region ?? _jupiterSettings.CurrentValue.CurrentSite;
 		MemoryBlobInfo? blobInfo = await GetBlobInfo(ns, blobIdentifier);
 		return blobInfo?.Regions.Contains(expectedRegion) ?? false;
 	}
 
-	public async IAsyncEnumerable<BaseBlobReference> GetBlobReferencesAsync(NamespaceId ns, BlobId id)
+	public async IAsyncEnumerable<BaseBlobReference> GetBlobReferencesAsync(NamespaceId ns, BlobId id, [EnumeratorCancellation] CancellationToken cancellationToken)
 	{
 		MemoryBlobInfo? blobInfo = await GetBlobInfo(ns, id);
 
@@ -116,7 +118,7 @@ public class MemoryBlobIndex : IBlobIndex
 		}
 	}
 
-	public Task AddRefToBlobsAsync(NamespaceId ns, BucketId bucket, RefId key, BlobId[] blobs)
+	public Task AddRefToBlobsAsync(NamespaceId ns, BucketId bucket, RefId key, BlobId[] blobs, CancellationToken cancellationToken)
 	{
 		foreach (BlobId id in blobs)
 		{
@@ -137,7 +139,7 @@ public class MemoryBlobIndex : IBlobIndex
 		return Task.CompletedTask;
 	}
 
-	public async IAsyncEnumerable<(NamespaceId, BlobId)> GetAllBlobsAsync()
+	public async IAsyncEnumerable<(NamespaceId, BlobId)> GetAllBlobsAsync([EnumeratorCancellation] CancellationToken cancellationToken)
 	{
 		await Task.CompletedTask;
 
@@ -150,7 +152,7 @@ public class MemoryBlobIndex : IBlobIndex
 		}
 	}
 
-	public Task RemoveReferencesAsync(NamespaceId ns, BlobId id, List<BaseBlobReference>? referencesToRemove)
+	public Task RemoveReferencesAsync(NamespaceId ns, BlobId id, List<BaseBlobReference>? referencesToRemove, CancellationToken cancellationToken)
 	{
 		ConcurrentDictionary<BlobId, MemoryBlobInfo> index = GetNamespaceContainer(ns);
 
@@ -172,7 +174,7 @@ public class MemoryBlobIndex : IBlobIndex
 		return Task.CompletedTask;
 	}
 
-	public async Task<List<string>> GetBlobRegionsAsync(NamespaceId ns, BlobId blob)
+	public async Task<List<string>> GetBlobRegionsAsync(NamespaceId ns, BlobId blob, CancellationToken cancellationToken)
 	{
 		MemoryBlobInfo? blobInfo = await GetBlobInfo(ns, blob);
 
@@ -184,7 +186,7 @@ public class MemoryBlobIndex : IBlobIndex
 		throw new BlobNotFoundException(ns, blob);
 	}
 
-	public async Task AddBlobReferencesAsync(NamespaceId ns, BlobId sourceBlob, BlobId targetBlob)
+	public async Task AddBlobReferencesAsync(NamespaceId ns, BlobId sourceBlob, BlobId targetBlob, CancellationToken cancellationToken)
 	{
 		MemoryBlobInfo? blobInfo = await GetBlobInfo(ns, sourceBlob);
 
@@ -196,7 +198,7 @@ public class MemoryBlobIndex : IBlobIndex
 		blobInfo.References.Add(new BlobToBlobReference(targetBlob));
 	}
 
-	public Task AddBlobToBucketListAsync(NamespaceId ns, BucketId bucket, RefId key, BlobId blobId, long blobSize)
+	public Task AddBlobToBucketListAsync(NamespaceId ns, BucketId bucket, RefId key, BlobId blobId, long blobSize, CancellationToken cancellationToken)
 	{
 		ConcurrentDictionary<BucketId, MemoryBucketInfo> bucketDict = _bucketIndex.GetOrAdd(ns, id => new ConcurrentDictionary<BucketId, MemoryBucketInfo>());
 		MemoryBucketInfo bucketInfo = bucketDict.GetOrAdd(bucket, id => new MemoryBucketInfo());
@@ -205,7 +207,7 @@ public class MemoryBlobIndex : IBlobIndex
 		return Task.CompletedTask;
 	}
 
-	public Task RemoveBlobFromBucketListAsync(NamespaceId ns, BucketId bucket, RefId key, List<BlobId> blobIds)
+	public Task RemoveBlobFromBucketListAsync(NamespaceId ns, BucketId bucket, RefId key, List<BlobId> blobIds, CancellationToken cancellationToken)
 	{
 		if (_bucketIndex.TryGetValue(ns, out ConcurrentDictionary<BucketId, MemoryBucketInfo>? bucketDict))
 		{
@@ -218,7 +220,7 @@ public class MemoryBlobIndex : IBlobIndex
 		return Task.CompletedTask;
 	}
 
-	public Task<BucketStats> CalculateBucketStatisticsAsync(NamespaceId ns, BucketId bucket)
+	public Task<BucketStats> CalculateBucketStatisticsAsync(NamespaceId ns, BucketId bucket, CancellationToken cancellationToken)
 	{
 		HashSet<RefId> foundRefs = new HashSet<RefId>();
 		int countOfBlobs = 0;
diff --git a/Engine/Source/Programs/UnrealCloudDDC/Jupiter/Implementation/Blob/MemoryBlobStore.cs b/Engine/Source/Programs/UnrealCloudDDC/Jupiter/Implementation/Blob/MemoryBlobStore.cs
--- a/Engine/Source/Programs/UnrealCloudDDC/Jupiter/Implementation/Blob/MemoryBlobStore.cs
+++ b/Engine/Source/Programs/UnrealCloudDDC/Jupiter/Implementation/Blob/MemoryBlobStore.cs
@@ -153,7 +153,7 @@ namespace Jupiter.Implementation
 
 		public async Task WriteAsync(string path, Stream stream, CancellationToken cancellationToken)
 		{
-			byte[] blob = await stream.ToByteArrayAsync();
+			byte[] blob = await stream.ToByteArrayAsync(cancellationToken);
 
 			// we do not split the blob into smaller parts when storing in memory, this is only for test purposes
 			// so there is no need to add that complexity
diff --git a/Engine/Source/Programs/UnrealCloudDDC/Jupiter/Implementation/Blob/MongoBlobIndex.cs b/Engine/Source/Programs/UnrealCloudDDC/Jupiter/Implementation/Blob/MongoBlobIndex.cs
--- a/Engine/Source/Programs/UnrealCloudDDC/Jupiter/Implementation/Blob/MongoBlobIndex.cs
+++ b/Engine/Source/Programs/UnrealCloudDDC/Jupiter/Implementation/Blob/MongoBlobIndex.cs
@@ -3,6 +3,8 @@
 using System;
 using System.Collections.Generic;
 using System.Linq;
+using System.Runtime.CompilerServices;
+using System.Threading;
 using System.Threading.Tasks;
 using EpicGames.Horde.Storage;
 using Microsoft.Extensions.Options;
@@ -38,11 +40,11 @@ public class MongoBlobIndex : MongoStore, IBlobIndex
 		});
 	}
 
-	private async Task<MongoBlobIndexModelV0?> GetBlobInfoAsync(NamespaceId ns, BlobId id)
+	private async Task<MongoBlobIndexModelV0?> GetBlobInfoAsync(NamespaceId ns, BlobId id, CancellationToken cancellationToken)
 	{
 		IMongoCollection<MongoBlobIndexModelV0> collection = GetCollection<MongoBlobIndexModelV0>();
-		IAsyncCursor<MongoBlobIndexModelV0>? cursor = await collection.FindAsync(m => m.Ns == ns.ToString() && m.BlobId == id.ToString());
-		MongoBlobIndexModelV0? model = await cursor.FirstOrDefaultAsync();
+		IAsyncCursor<MongoBlobIndexModelV0>? cursor = await collection.FindAsync(m => m.Ns == ns.ToString() && m.BlobId == id.ToString(), cancellationToken: cancellationToken);
+		MongoBlobIndexModelV0? model = await cursor.FirstOrDefaultAsync(cancellationToken);
 		if (model == null)
 		{
 			return null;
@@ -51,7 +53,7 @@ public class MongoBlobIndex : MongoStore, IBlobIndex
 		return model;
 	}
 
-	public async Task AddBlobToIndexAsync(NamespaceId ns, BlobId id, string? region = null)
+	public async Task AddBlobToIndexAsync(NamespaceId ns, BlobId id, string? region = null, CancellationToken cancellationToken = default)
 	{
 		region ??= _jupiterSettings.CurrentValue.CurrentSite;
 		IMongoCollection<MongoBlobIndexModelV0> collection = GetCollection<MongoBlobIndexModelV0>();
@@ -62,14 +64,14 @@ public class MongoBlobIndex : MongoStore, IBlobIndex
 		await collection.FindOneAndReplaceAsync(filter, model, new FindOneAndReplaceOptions<MongoBlobIndexModelV0, MongoBlobIndexModelV0>
 		{
 			IsUpsert = true
-		});
+		}, cancellationToken);
 	}
 
-	public async Task RemoveBlobFromRegionAsync(NamespaceId ns, BlobId id, string? region = null)
+	public async Task RemoveBlobFromRegionAsync(NamespaceId ns, BlobId id, string? region = null, CancellationToken cancellationToken = default)
 	{
 		region ??= _jupiterSettings.CurrentValue.CurrentSite;
 
-		MongoBlobIndexModelV0? model = await GetBlobInfoAsync(ns, id);
+		MongoBlobIndexModelV0? model = await GetBlobInfoAsync(ns, id, cancellationToken);
 		if (model == null)
 		{
 			throw new BlobNotFoundException(ns, id);
@@ -80,13 +82,13 @@ public class MongoBlobIndex : MongoStore, IBlobIndex
 		await collection.FindOneAndReplaceAsync(filter, model, new FindOneAndReplaceOptions<MongoBlobIndexModelV0, MongoBlobIndexModelV0>
 		{
 			IsUpsert = true
-		});
+		}, cancellationToken);
 
 	}
 
-	public async IAsyncEnumerable<BaseBlobReference> GetBlobReferencesAsync(NamespaceId ns, BlobId id)
+	public async IAsyncEnumerable<BaseBlobReference> GetBlobReferencesAsync(NamespaceId ns, BlobId id, [EnumeratorCancellation] CancellationToken cancellationToken)
 	{
-		MongoBlobIndexModelV0? blobInfo = await GetBlobInfoAsync(ns, id);
+		MongoBlobIndexModelV0? blobInfo = await GetBlobInfoAsync(ns, id, cancellationToken);
 		if (blobInfo == null)
 		{
 			yield break;
@@ -112,13 +114,13 @@ public class MongoBlobIndex : MongoStore, IBlobIndex
 		}
 	}
 
-	public async Task<bool> BlobExistsInRegionAsync(NamespaceId ns, BlobId blobIdentifier, string? region = null)
+	public async Task<bool> BlobExistsInRegionAsync(NamespaceId ns, BlobId blobIdentifier, string? region = null, CancellationToken cancellationToken = default)
 	{
-		MongoBlobIndexModelV0? blobInfo = await GetBlobInfoAsync(ns, blobIdentifier);
+		MongoBlobIndexModelV0? blobInfo = await GetBlobInfoAsync(ns, blobIdentifier, cancellationToken);
 		return blobInfo?.Regions.Contains(_jupiterSettings.CurrentValue.CurrentSite) ?? false;
 	}
 
-	public async Task AddRefToBlobsAsync(NamespaceId ns, BucketId bucket, RefId key, BlobId[] blobs)
+	public async Task AddRefToBlobsAsync(NamespaceId ns, BucketId bucket, RefId key, BlobId[] blobs, CancellationToken cancellationToken)
 	{
 		IMongoCollection<MongoBlobIndexModelV0> collection = GetCollection<MongoBlobIndexModelV0>();
 
@@ -132,19 +134,19 @@ public class MongoBlobIndex : MongoStore, IBlobIndex
 				UpdateDefinition<MongoBlobIndexModelV0> update = Builders<MongoBlobIndexModelV0>.Update.AddToSet(m => m.References, new Dictionary<string, string> { { "bucket", bucket.ToString() }, { "key", key.ToString() } });
 				FilterDefinition<MongoBlobIndexModelV0> filter = Builders<MongoBlobIndexModelV0>.Filter.Where(m => m.Ns == nsAsString && m.BlobId == id.ToString());
 
-				await collection.FindOneAndUpdateAsync(filter, update);
-			});
+				await collection.FindOneAndUpdateAsync(filter, update, cancellationToken: cancellationToken);
+			}, cancellationToken);
 		}
 
 		await Task.WhenAll(refUpdateTasks);
 	}
 
-	public async IAsyncEnumerable<(NamespaceId, BlobId)> GetAllBlobsAsync()
+	public async IAsyncEnumerable<(NamespaceId, BlobId)> GetAllBlobsAsync([EnumeratorCancellation] CancellationToken cancellationToken)
 	{
 		IMongoCollection<MongoBlobIndexModelV0> collection = GetCollection<MongoBlobIndexModelV0>();
-		IAsyncCursor<MongoBlobIndexModelV0>? cursor = await collection.FindAsync(FilterDefinition<MongoBlobIndexModelV0>.Empty);
+		IAsyncCursor<MongoBlobIndexModelV0>? cursor = await collection.FindAsync(FilterDefinition<MongoBlobIndexModelV0>.Empty, cancellationToken: cancellationToken);
 
-		while (await cursor.MoveNextAsync())
+		while (await cursor.MoveNextAsync(cancellationToken))
 		{
 			foreach (MongoBlobIndexModelV0 model in cursor.Current)
 			{
@@ -153,7 +155,7 @@ public class MongoBlobIndex : MongoStore, IBlobIndex
 		}
 	}
 
-	public async Task RemoveReferencesAsync(NamespaceId ns, BlobId id, List<BaseBlobReference>? referencesToRemove)
+	public async Task RemoveReferencesAsync(NamespaceId ns, BlobId id, List<BaseBlobReference>? referencesToRemove, CancellationToken cancellationToken)
 	{
 		IMongoCollection<MongoBlobIndexModelV0> collection = GetCollection<MongoBlobIndexModelV0>();
 
@@ -161,7 +163,7 @@ public class MongoBlobIndex : MongoStore, IBlobIndex
 		if (referencesToRemove == null)
 		{
 			FilterDefinition<MongoBlobIndexModelV0> filter = Builders<MongoBlobIndexModelV0>.Filter.Where(m => m.Ns == nsAsString && m.BlobId == id.ToString());
-			await collection.DeleteOneAsync(filter);
+			await collection.DeleteOneAsync(filter, cancellationToken);
 		}
 		else
 		{
@@ -189,13 +191,13 @@ public class MongoBlobIndex : MongoStore, IBlobIndex
 			UpdateDefinition<MongoBlobIndexModelV0> update = Builders<MongoBlobIndexModelV0>.Update.PullAll(m => m.References, refs);
 			FilterDefinition<MongoBlobIndexModelV0> filter = Builders<MongoBlobIndexModelV0>.Filter.Where(m => m.Ns == nsAsString && m.BlobId == id.ToString());
 
-			await collection.FindOneAndUpdateAsync(filter, update);
+			await collection.FindOneAndUpdateAsync(filter, update, cancellationToken: cancellationToken);
 		}
 	}
 
-	public async Task<List<string>> GetBlobRegionsAsync(NamespaceId ns, BlobId blob)
+	public async Task<List<string>> GetBlobRegionsAsync(NamespaceId ns, BlobId blob, CancellationToken cancellationToken)
 	{
-		MongoBlobIndexModelV0? blobInfo = await GetBlobInfoAsync(ns, blob);
+		MongoBlobIndexModelV0? blobInfo = await GetBlobInfoAsync(ns, blob, cancellationToken);
 		if (blobInfo == null)
 		{
 			return new List<string>();
@@ -203,7 +205,7 @@ public class MongoBlobIndex : MongoStore, IBlobIndex
 		return blobInfo.Regions.ToList();
 	}
 
-	public async Task AddBlobReferencesAsync(NamespaceId ns, BlobId sourceBlob, BlobId targetBlob)
+	public async Task AddBlobReferencesAsync(NamespaceId ns, BlobId sourceBlob, BlobId targetBlob, CancellationToken cancellationToken)
 	{
 		IMongoCollection<MongoBlobIndexModelV0> collection = GetCollection<MongoBlobIndexModelV0>();
 
@@ -212,26 +214,26 @@ public class MongoBlobIndex : MongoStore, IBlobIndex
 		UpdateDefinition<MongoBlobIndexModelV0> update = Builders<MongoBlobIndexModelV0>.Update.AddToSet(m => m.References, new Dictionary<string, string> { { "blob_id", targetBlob.ToString() } });
 		FilterDefinition<MongoBlobIndexModelV0> filter = Builders<MongoBlobIndexModelV0>.Filter.Where(m => m.Ns == nsAsString && m.BlobId == sourceBlob.ToString());
 
-		await collection.FindOneAndUpdateAsync(filter, update);
+		await collection.FindOneAndUpdateAsync(filter, update, cancellationToken: cancellationToken);
 	}
 
-	public async Task AddBlobToBucketListAsync(NamespaceId ns, BucketId bucket, RefId key, BlobId blobId, long blobSize)
+	public async Task AddBlobToBucketListAsync(NamespaceId ns, BucketId bucket, RefId key, BlobId blobId, long blobSize, CancellationToken cancellationToken)
 	{
 		IMongoCollection<MongoBucketBlobV0> collection = GetCollection<MongoBucketBlobV0>();
 
 		FilterDefinition<MongoBucketBlobV0> filter = Builders<MongoBucketBlobV0>.Filter.Where(m => m.Ns == ns.ToString() && m.BucketId == bucket.ToString() && m.RefId == key.ToString() && m.BlobId == blobId.ToString());
-		await collection.ReplaceOneAsync(filter, new MongoBucketBlobV0(ns, bucket, key, blobId, blobSize), new ReplaceOptions() { IsUpsert = true });
+		await collection.ReplaceOneAsync(filter, new MongoBucketBlobV0(ns, bucket, key, blobId, blobSize), new ReplaceOptions() { IsUpsert = true }, cancellationToken);
 	}
 
-	public async Task RemoveBlobFromBucketListAsync(NamespaceId ns, BucketId bucket, RefId key, List<BlobId> blobIds)
+	public async Task RemoveBlobFromBucketListAsync(NamespaceId ns, BucketId bucket, RefId key, List<BlobId> blobIds, CancellationToken cancellationToken)
 	{
 		IMongoCollection<MongoBucketBlobV0> collection = GetCollection<MongoBucketBlobV0>();
 
 		FilterDefinition<MongoBucketBlobV0> filter = Builders<MongoBucketBlobV0>.Filter.Where(m => m.Ns == ns.ToString() && m.BucketId == bucket.ToString() && m.RefId == key.ToString());
-		await collection.DeleteManyAsync(filter);
+		await collection.DeleteManyAsync(filter, cancellationToken);
 	}
 
-	public async Task<BucketStats> CalculateBucketStatisticsAsync(NamespaceId ns, BucketId bucket)
+	public async Task<BucketStats> CalculateBucketStatisticsAsync(NamespaceId ns, BucketId bucket, CancellationToken cancellationToken)
 	{
 		IMongoCollection<MongoBucketBlobV0> collection = GetCollection<MongoBucketBlobV0>();
 		FilterDefinition<MongoBucketBlobV0> filter = Builders<MongoBucketBlobV0>.Filter.Where(m => m.Ns == ns.ToString() && m.BucketId == bucket.ToString());
@@ -246,7 +248,7 @@ public class MongoBlobIndex : MongoStore, IBlobIndex
 					SmallestBlob = r.Min(a => a.Size),
 					LargestBlob = r.Max(a => a.Size),
 					CountOfBlobs = r.Count(),
-				}).ToListAsync();
+				}).ToListAsync(cancellationToken);
 
 		var blobStat = blobStats.FirstOrDefault();
 
@@ -256,7 +258,7 @@ public class MongoBlobIndex : MongoStore, IBlobIndex
 			.Group(m => m.RefId,
 				grouping => new { DoesNotMatter = grouping.Key })
 			.Count()
-			.First()
+			.First(cancellationToken)
 			.Count;
 
 		return new BucketStats()
diff --git a/Engine/Source/Programs/UnrealCloudDDC/Jupiter/Implementation/Blob/RefStoreConsistencyCheckService.cs b/Engine/Source/Programs/UnrealCloudDDC/Jupiter/Implementation/Blob/RefStoreConsistencyCheckService.cs
--- a/Engine/Source/Programs/UnrealCloudDDC/Jupiter/Implementation/Blob/RefStoreConsistencyCheckService.cs
+++ b/Engine/Source/Programs/UnrealCloudDDC/Jupiter/Implementation/Blob/RefStoreConsistencyCheckService.cs
@@ -47,7 +47,7 @@ namespace Jupiter.Implementation
 			_policyResolver = policyResolver;
 		}
 
-		public override async Task<bool> OnPollAsync(ConsistencyState state, CancellationToken cancellationToken)
+		public override async Task<bool> OnPollAsync(ConsistencyState state, CancellationToken cancellationToken = default)
 		{
 			if (!_settings.CurrentValue.EnableRefStoreChecks)
 			{
@@ -55,12 +55,12 @@ namespace Jupiter.Implementation
 				return false;
 			}
 
-			await RunConsistencyCheckAsync();
+			await RunConsistencyCheckAsync(cancellationToken);
 
 			return true;
 		}
 
-		public async Task RunConsistencyCheckAsync()
+		public async Task RunConsistencyCheckAsync(CancellationToken cancellationToken = default)
 		{
 			if (!_leaderElection.IsThisInstanceLeader())
 			{
@@ -70,17 +70,17 @@ namespace Jupiter.Implementation
 
 			ulong countOfRefsChecked = 0;
 			ulong countOfMissingLastAccessTime = 0;
-			await foreach ((NamespaceId ns, BucketId bucket, RefId refId) in _referencesStore.GetRecordsWithoutAccessTimeAsync())
+			await foreach ((NamespaceId ns, BucketId bucket, RefId refId) in _referencesStore.GetRecordsWithoutAccessTimeAsync(cancellationToken))
 			{
 				using TelemetrySpan scope = _tracer.StartActiveSpan("consistency_check.ref_store")
 					.SetAttribute("operation.name", "consistency_check.ref_store")
 					.SetAttribute("resource.name", $"{ns}.{bucket}.{refId}");
 
-				DateTime? lastAccessTime = await _referencesStore.GetLastAccessTimeAsync(ns, bucket, refId);
+				DateTime? lastAccessTime = await _referencesStore.GetLastAccessTimeAsync(ns, bucket, refId, cancellationToken);
 				if (!lastAccessTime.HasValue)
 				{
 					// if there is no last access time record we add one so that the two tables are consistent, this will make the ref record be considered by the GC and thus cleanup anything that might be very old (but its added as a new record so it will take until the configured cleanup time has passed)
-					await _referencesStore.UpdateLastAccessTimeAsync(ns, bucket, refId, DateTime.Now);
+					await _referencesStore.UpdateLastAccessTimeAsync(ns, bucket, refId, DateTime.Now, cancellationToken);
 
 					Interlocked.Increment(ref countOfMissingLastAccessTime);
 				}
diff --git a/Engine/Source/Programs/UnrealCloudDDC/Jupiter/Implementation/Blob/ScyllaBlobIndex.cs b/Engine/Source/Programs/UnrealCloudDDC/Jupiter/Implementation/Blob/ScyllaBlobIndex.cs
--- a/Engine/Source/Programs/UnrealCloudDDC/Jupiter/Implementation/Blob/ScyllaBlobIndex.cs
+++ b/Engine/Source/Programs/UnrealCloudDDC/Jupiter/Implementation/Blob/ScyllaBlobIndex.cs
@@ -4,6 +4,7 @@ using System;
 using System.Collections.Generic;
 using System.Diagnostics;
 using System.Linq;
+using System.Runtime.CompilerServices;
 using System.Threading;
 using System.Threading.Tasks;
 using Cassandra;
@@ -91,7 +92,7 @@ public class ScyllaBlobIndex : IBlobIndex
 		_getBucketStatsRefsStatement = _session.Prepare("select count(reference_id) from bucket_referenced_ref WHERE namespace = ? AND bucket_id = ?  AND hash_prefix = ?");
 	}
 
-	public async Task AddBlobToIndexAsync(NamespaceId ns, BlobId id, string? region = null)
+	public async Task AddBlobToIndexAsync(NamespaceId ns, BlobId id, string? region = null, CancellationToken cancellationToken = default)
 	{
 		region ??= _jupiterSettings.CurrentValue.CurrentSite;
 		using TelemetrySpan scope = _tracer.BuildScyllaSpan("scylla.insert_blob_index").SetAttribute("resource.name", $"{ns}.{id}");
@@ -135,7 +136,7 @@ public class ScyllaBlobIndex : IBlobIndex
 		}
 	}
 
-	public async Task RemoveBlobFromRegionAsync(NamespaceId ns, BlobId id, string? region = null)
+	public async Task RemoveBlobFromRegionAsync(NamespaceId ns, BlobId id, string? region = null, CancellationToken cancellationToken = default)
 	{
 		region ??= _jupiterSettings.CurrentValue.CurrentSite;
 		using TelemetrySpan scope = _tracer.BuildScyllaSpan("scylla.remove_blob_index_region").SetAttribute("resource.name", $"{ns}.{id}");
@@ -143,7 +144,7 @@ public class ScyllaBlobIndex : IBlobIndex
 		await _mapper.DeleteAsync<ScyllaBlobIndexEntry>(new ScyllaBlobIndexEntry(ns.ToString(), id, region));
 	}
 
-	public async Task<bool> BlobExistsInRegionAsync(NamespaceId ns, BlobId blobIdentifier, string? region = null)
+	public async Task<bool> BlobExistsInRegionAsync(NamespaceId ns, BlobId blobIdentifier, string? region = null, CancellationToken cancellationToken = default)
 	{
 		region ??= _jupiterSettings.CurrentValue.CurrentSite;
 		ScyllaBlobIndexEntry? entry = await _mapper.SingleOrDefaultAsync<ScyllaBlobIndexEntry>("WHERE namespace = ? AND blob_id = ? AND region = ?", ns.ToString(), blobIdentifier.HashData, region);
@@ -160,7 +161,7 @@ public class ScyllaBlobIndex : IBlobIndex
 
 			foreach (string oldRegion in regions)
 			{
-				await AddBlobToIndexAsync(ns, blobIdentifier, oldRegion);
+				await AddBlobToIndexAsync(ns, blobIdentifier, oldRegion, cancellationToken);
 			}
 
 			// blob has been migrated so it existed
@@ -169,7 +170,7 @@ public class ScyllaBlobIndex : IBlobIndex
 
 		return !blobMissing;
 	}
-	public async IAsyncEnumerable<(NamespaceId, BlobId)> GetAllBlobsAsync()
+	public async IAsyncEnumerable<(NamespaceId, BlobId)> GetAllBlobsAsync([EnumeratorCancellation] CancellationToken cancellationToken)
 	{
 		using TelemetrySpan scope = _tracer.BuildScyllaSpan("scylla.get_all_blobs");
 
@@ -181,7 +182,7 @@ public class ScyllaBlobIndex : IBlobIndex
 		}
 	}
 
-	public async Task<List<string>> GetBlobRegionsAsync(NamespaceId ns, BlobId blob)
+	public async Task<List<string>> GetBlobRegionsAsync(NamespaceId ns, BlobId blob, CancellationToken cancellationToken)
 	{
 		List<string> regions = new List<string>();
 		foreach (ScyllaBlobIndexEntry blobIndex in await _mapper.FetchAsync<ScyllaBlobIndexEntry>("WHERE namespace = ? AND blob_id = ?", ns.ToString(), blob.HashData))
@@ -200,7 +201,7 @@ public class ScyllaBlobIndex : IBlobIndex
 
 			foreach (string oldRegion in oldRegions)
 			{
-				await AddBlobToIndexAsync(ns, blob, oldRegion);
+				await AddBlobToIndexAsync(ns, blob, oldRegion, cancellationToken);
 			}
 
 			// blob has been migrated lets return the old region list
@@ -209,7 +210,7 @@ public class ScyllaBlobIndex : IBlobIndex
 		return regions;
 	}
 
-	public async Task AddBlobReferencesAsync(NamespaceId ns, BlobId sourceBlob, BlobId targetBlob)
+	public async Task AddBlobReferencesAsync(NamespaceId ns, BlobId sourceBlob, BlobId targetBlob, CancellationToken cancellationToken)
 	{
 		using TelemetrySpan scope = _tracer.BuildScyllaSpan("scylla.add_blob_to_blob_ref");
 
@@ -219,7 +220,7 @@ public class ScyllaBlobIndex : IBlobIndex
 		await _mapper.InsertAsync<ScyllaBlobIncomingReference>(incomingReference);
 	}
 
-	public async IAsyncEnumerable<BaseBlobReference> GetBlobReferencesAsync(NamespaceId ns, BlobId id)
+	public async IAsyncEnumerable<BaseBlobReference> GetBlobReferencesAsync(NamespaceId ns, BlobId id, [EnumeratorCancellation] CancellationToken cancellationToken)
 	{
 		using TelemetrySpan scope = _tracer.BuildScyllaSpan("scylla.get_blob_references").SetAttribute("resource.name", $"{ns}.{id}");
 
@@ -258,14 +259,14 @@ public class ScyllaBlobIndex : IBlobIndex
 				{
 					BucketId bucket = new BucketId(scyllaObjectReference.Bucket);
 					RefId key = new RefId(scyllaObjectReference.Key);
-					await AddRefToBlobsAsync(ns, bucket, key, new[] { id });
+					await AddRefToBlobsAsync(ns, bucket, key, new[] { id }, cancellationToken);
 					yield return new RefBlobReference(bucket, key);
 				}
 			}
 		}
 	}
 
-	public async Task AddRefToBlobsAsync(NamespaceId ns, BucketId bucket, RefId key, BlobId[] blobs)
+	public async Task AddRefToBlobsAsync(NamespaceId ns, BucketId bucket, RefId key, BlobId[] blobs, CancellationToken cancellationToken)
 	{
 		using TelemetrySpan scope = _tracer.BuildScyllaSpan("scylla.add_ref_blobs");
 
@@ -282,7 +283,7 @@ public class ScyllaBlobIndex : IBlobIndex
 		await Task.WhenAll(refUpdateTasks);
 	}
 
-	public async Task RemoveReferencesAsync(NamespaceId ns, BlobId id, List<BaseBlobReference>? referencesToRemove)
+	public async Task RemoveReferencesAsync(NamespaceId ns, BlobId id, List<BaseBlobReference>? referencesToRemove, CancellationToken cancellationToken)
 	{
 		using TelemetrySpan scope = _tracer.BuildScyllaSpan("scylla.remove_ref_blobs");
 
@@ -321,7 +322,7 @@ public class ScyllaBlobIndex : IBlobIndex
 		}
 	}
 
-	public async Task AddBlobToBucketListAsync(NamespaceId ns, BucketId bucket, RefId key, BlobId blobId, long blobSize)
+	public async Task AddBlobToBucketListAsync(NamespaceId ns, BucketId bucket, RefId key, BlobId blobId, long blobSize, CancellationToken cancellationToken)
 	{
 		using TelemetrySpan scope = _tracer.BuildScyllaSpan("scylla.add_bucket_blob");
 
@@ -343,7 +344,7 @@ public class ScyllaBlobIndex : IBlobIndex
 		await Task.WhenAll(insertRefRefTask, insertBlobRefTask);
 	}
 
-	public async Task RemoveBlobFromBucketListAsync(NamespaceId ns, BucketId bucket, RefId key, List<BlobId> blobIds)
+	public async Task RemoveBlobFromBucketListAsync(NamespaceId ns, BucketId bucket, RefId key, List<BlobId> blobIds, CancellationToken cancellationToken)
 	{
 		using TelemetrySpan scope = _tracer.BuildScyllaSpan("scylla.remove_bucket_blob");
 
@@ -362,7 +363,7 @@ public class ScyllaBlobIndex : IBlobIndex
 		await deleteRefRefTask;
 	}
 
-	public async Task<BucketStats> CalculateBucketStatisticsAsync(NamespaceId ns, BucketId bucket)
+	public async Task<BucketStats> CalculateBucketStatisticsAsync(NamespaceId ns, BucketId bucket, CancellationToken cancellationToken)
 	{
 		string nsAsString = ns.ToString();
 		string bucketAsString = bucket.ToString();
diff --git a/Engine/Source/Programs/UnrealCloudDDC/Jupiter/Implementation/ContentId/CacheContentIdStore.cs b/Engine/Source/Programs/UnrealCloudDDC/Jupiter/Implementation/ContentId/CacheContentIdStore.cs
--- a/Engine/Source/Programs/UnrealCloudDDC/Jupiter/Implementation/ContentId/CacheContentIdStore.cs
+++ b/Engine/Source/Programs/UnrealCloudDDC/Jupiter/Implementation/ContentId/CacheContentIdStore.cs
@@ -4,6 +4,7 @@ using System;
 using System.Net;
 using System.Net.Http;
 using System.Net.Http.Json;
+using System.Threading;
 using System.Threading.Tasks;
 using EpicGames.Horde.Storage;
 using Jupiter.Controllers;
@@ -18,10 +19,10 @@ namespace Jupiter.Implementation
 		{
 		}
 
-		public async Task<BlobId[]?> ResolveAsync(NamespaceId ns, ContentId contentId, bool mustBeContentId)
+		public async Task<BlobId[]?> ResolveAsync(NamespaceId ns, ContentId contentId, bool mustBeContentId, CancellationToken cancellationToken)
 		{
 			using HttpRequestMessage getContentIdRequest = await BuildHttpRequestAsync(HttpMethod.Get, new Uri($"api/v1/content-id/{ns}/{contentId}", UriKind.Relative));
-			HttpResponseMessage response = await HttpClient.SendAsync(getContentIdRequest);
+			HttpResponseMessage response = await HttpClient.SendAsync(getContentIdRequest, cancellationToken);
 
 			if (response.StatusCode == HttpStatusCode.NotFound)
 			{
@@ -29,7 +30,7 @@ namespace Jupiter.Implementation
 			}
 
 			response.EnsureSuccessStatusCode();
-			ResolvedContentIdResponse? resolvedContentId = await response.Content.ReadFromJsonAsync<ResolvedContentIdResponse>();
+			ResolvedContentIdResponse? resolvedContentId = await response.Content.ReadFromJsonAsync<ResolvedContentIdResponse>(cancellationToken);
 
 			if (resolvedContentId == null)
 			{
@@ -38,10 +39,10 @@ namespace Jupiter.Implementation
 			return resolvedContentId.Blobs;
 		}
 
-		public async Task PutAsync(NamespaceId ns, ContentId contentId, BlobId blobIdentifier, int contentWeight)
+		public async Task PutAsync(NamespaceId ns, ContentId contentId, BlobId blobIdentifier, int contentWeight, CancellationToken cancellationToken)
 		{
 			using HttpRequestMessage putContentIdRequest = await BuildHttpRequestAsync(HttpMethod.Put, new Uri($"api/v1/content-id/{ns}/{contentId}/update/{blobIdentifier}/{contentWeight}", UriKind.Relative));
-			HttpResponseMessage response = await HttpClient.SendAsync(putContentIdRequest);
+			HttpResponseMessage response = await HttpClient.SendAsync(putContentIdRequest, cancellationToken);
 
 			response.EnsureSuccessStatusCode();
 		}
diff --git a/Engine/Source/Programs/UnrealCloudDDC/Jupiter/Implementation/ContentId/CompressedBufferUtils.cs b/Engine/Source/Programs/UnrealCloudDDC/Jupiter/Implementation/ContentId/CompressedBufferUtils.cs
--- a/Engine/Source/Programs/UnrealCloudDDC/Jupiter/Implementation/ContentId/CompressedBufferUtils.cs
+++ b/Engine/Source/Programs/UnrealCloudDDC/Jupiter/Implementation/ContentId/CompressedBufferUtils.cs
@@ -6,6 +6,7 @@ using System.Collections.Generic;
 using System.IO;
 using System.Linq;
 using System.Text;
+using System.Threading;
 using System.Threading.Tasks;
 using Blake3;
 using EpicGames.Compression;
@@ -183,7 +184,7 @@ namespace Jupiter.Implementation
 			}
 		}
 
-		public async Task<IBufferedPayload> DecompressContentAsync(Stream sourceStream, ulong streamSize)
+		public async Task<IBufferedPayload> DecompressContentAsync(Stream sourceStream, ulong streamSize, CancellationToken cancellationToken = default)
 		{
 			using BinaryReader br = new BinaryReader(sourceStream);
 			(CompressedBufferHeader header, uint[] compressedBlockSizes) = ExtractHeader(br);
@@ -226,7 +227,7 @@ namespace Jupiter.Implementation
 				}
 				else
 				{
-					await sourceStream.CopyToAsync(targetStream);
+					await sourceStream.CopyToAsync(targetStream, cancellationToken);
 				}
 			}
 
@@ -250,7 +251,7 @@ namespace Jupiter.Implementation
 
 					BlobId headerIdentifier = new BlobId(slicedHash);
 					await using Stream hashStream = finalizedBufferedPayload.GetStream();
-					BlobId contentHash = await BlobId.FromStreamAsync(hashStream);
+					BlobId contentHash = await BlobId.FromStreamAsync(hashStream, cancellationToken);
 
 					if (!headerIdentifier.Equals(contentHash))
 					{
diff --git a/Engine/Source/Programs/UnrealCloudDDC/Jupiter/Implementation/ContentId/IContentIdStore.cs b/Engine/Source/Programs/UnrealCloudDDC/Jupiter/Implementation/ContentId/IContentIdStore.cs
--- a/Engine/Source/Programs/UnrealCloudDDC/Jupiter/Implementation/ContentId/IContentIdStore.cs
+++ b/Engine/Source/Programs/UnrealCloudDDC/Jupiter/Implementation/ContentId/IContentIdStore.cs
@@ -1,6 +1,7 @@
 // Copyright Epic Games, Inc. All Rights Reserved.
 
 using System;
+using System.Threading;
 using System.Threading.Tasks;
 using EpicGames.Horde.Storage;
 
@@ -14,8 +15,9 @@ namespace Jupiter.Implementation
 		/// <param name="ns">The namespace to operate in</param>
 		/// <param name="contentId">The identifier for the content id</param>
 		/// <param name="mustBeContentId"></param>
+		/// <param name="cancellationToken">Cancellation token for the operation</param>
 		/// <returns></returns>
-		Task<BlobId[]?> ResolveAsync(NamespaceId ns, ContentId contentId, bool mustBeContentId = false);
+		Task<BlobId[]?> ResolveAsync(NamespaceId ns, ContentId contentId, bool mustBeContentId = false, CancellationToken cancellationToken = default);
 
 		/// <summary>
 		/// Add a mapping from contentId to blobIdentifier
@@ -24,8 +26,9 @@ namespace Jupiter.Implementation
 		/// <param name="contentId">The contentId</param>
 		/// <param name="blobIdentifier">The blob the content id maps to</param>
 		/// <param name="contentWeight">Weight of this identifier compared to previous mappings, used to determine which is more important, lower weight is considered a better fit</param>
+		/// <param name="cancellationToken">Cancellation token for the operation</param>
 		/// <returns></returns>
-		Task PutAsync(NamespaceId ns, ContentId contentId, BlobId blobIdentifier, int contentWeight);
+		Task PutAsync(NamespaceId ns, ContentId contentId, BlobId blobIdentifier, int contentWeight, CancellationToken cancellationToken = default);
 	}
 
 	public class InvalidContentIdException : Exception
diff --git a/Engine/Source/Programs/UnrealCloudDDC/Jupiter/Implementation/ContentId/MemoryCachedContentIdStore.cs b/Engine/Source/Programs/UnrealCloudDDC/Jupiter/Implementation/ContentId/MemoryCachedContentIdStore.cs
--- a/Engine/Source/Programs/UnrealCloudDDC/Jupiter/Implementation/ContentId/MemoryCachedContentIdStore.cs
+++ b/Engine/Source/Programs/UnrealCloudDDC/Jupiter/Implementation/ContentId/MemoryCachedContentIdStore.cs
@@ -2,6 +2,7 @@
 
 using System;
 using System.Collections.Concurrent;
+using System.Threading;
 using System.Threading.Tasks;
 using EpicGames.Horde.Storage;
 using Microsoft.Extensions.Caching.Memory;
@@ -31,7 +32,7 @@ namespace Jupiter.Implementation
 			_contentIdCaches.Clear();
 		}
 
-		public async Task<BlobId[]?> ResolveAsync(NamespaceId ns, ContentId contentId, bool mustBeContentId)
+		public async Task<BlobId[]?> ResolveAsync(NamespaceId ns, ContentId contentId, bool mustBeContentId, CancellationToken cancellationToken)
 		{
 			MemoryCache cache = GetCacheForNamespace(ns);
 			if (cache.TryGetValue(contentId, out CachedContentIdEntry? cachedResult))
@@ -39,7 +40,7 @@ namespace Jupiter.Implementation
 				return cachedResult?.ReferencedBlobs;
 			}
 
-			BlobId[]? referencedBlobs = await _actualContentIdStore.ResolveAsync(ns, contentId, mustBeContentId);
+			BlobId[]? referencedBlobs = await _actualContentIdStore.ResolveAsync(ns, contentId, mustBeContentId, cancellationToken);
 			// we do not cache unresolvable content ids
 			if (referencedBlobs == null)
 			{
@@ -51,9 +52,9 @@ namespace Jupiter.Implementation
 			return referencedBlobs;
 		}
 
-		public Task PutAsync(NamespaceId ns, ContentId contentId, BlobId blobIdentifier, int contentWeight)
+		public Task PutAsync(NamespaceId ns, ContentId contentId, BlobId blobIdentifier, int contentWeight, CancellationToken cancellationToken)
 		{
-			Task actualPutTask = _actualContentIdStore.PutAsync(ns, contentId, blobIdentifier, contentWeight);
+			Task actualPutTask = _actualContentIdStore.PutAsync(ns, contentId, blobIdentifier, contentWeight, cancellationToken);
 			AddCacheEntry(ns, contentId, new BlobId[] { blobIdentifier });
 			return actualPutTask;
 		}
diff --git a/Engine/Source/Programs/UnrealCloudDDC/Jupiter/Implementation/ContentId/MemoryContentIdStore.cs b/Engine/Source/Programs/UnrealCloudDDC/Jupiter/Implementation/ContentId/MemoryContentIdStore.cs
--- a/Engine/Source/Programs/UnrealCloudDDC/Jupiter/Implementation/ContentId/MemoryContentIdStore.cs
+++ b/Engine/Source/Programs/UnrealCloudDDC/Jupiter/Implementation/ContentId/MemoryContentIdStore.cs
@@ -2,6 +2,7 @@
 
 using System.Collections.Concurrent;
 using System.Collections.Generic;
+using System.Threading;
 using System.Threading.Tasks;
 using EpicGames.Horde.Storage;
 
@@ -17,7 +18,7 @@ namespace Jupiter.Implementation
 			_blobStore = blobStore;
 		}
 
-		public async Task<BlobId[]?> ResolveAsync(NamespaceId ns, ContentId contentId, bool mustBeContentId)
+		public async Task<BlobId[]?> ResolveAsync(NamespaceId ns, ContentId contentId, bool mustBeContentId, CancellationToken cancellationToken)
 		{
 			if (_contentIds.TryGetValue(ns, out ConcurrentDictionary<ContentId, SortedList<int, BlobId[]>>? contentIdsForNamespace))
 			{
@@ -25,7 +26,7 @@ namespace Jupiter.Implementation
 				{
 					foreach ((int _, BlobId[] blobs) in contentIdMappings)
 					{
-						BlobId[] missingBlobs = await _blobStore.FilterOutKnownBlobsAsync(ns, blobs);
+						BlobId[] missingBlobs = await _blobStore.FilterOutKnownBlobsAsync(ns, blobs, cancellationToken);
 						if (missingBlobs.Length == 0)
 						{
 							return blobs;
@@ -37,7 +38,7 @@ namespace Jupiter.Implementation
 
 			BlobId uncompressedBlobIdentifier = contentId.AsBlobIdentifier();
 			// if no content id is found, but we have a blob that matches the content id (so a unchunked and uncompressed version of the data) we use that instead
-			if (!mustBeContentId && await _blobStore.ExistsAsync(ns, uncompressedBlobIdentifier))
+			if (!mustBeContentId && await _blobStore.ExistsAsync(ns, uncompressedBlobIdentifier, cancellationToken: cancellationToken))
 			{
 				return new[] { uncompressedBlobIdentifier };
 			}
@@ -45,7 +46,7 @@ namespace Jupiter.Implementation
 			return null;
 		}
 
-		public Task PutAsync(NamespaceId ns, ContentId contentId, BlobId blobIdentifier, int contentWeight)
+		public Task PutAsync(NamespaceId ns, ContentId contentId, BlobId blobIdentifier, int contentWeight, CancellationToken cancellationToken)
 		{
 			_contentIds.AddOrUpdate(ns, (_) =>
 			{
diff --git a/Engine/Source/Programs/UnrealCloudDDC/Jupiter/Implementation/ContentId/MongoContentIdStore.cs b/Engine/Source/Programs/UnrealCloudDDC/Jupiter/Implementation/ContentId/MongoContentIdStore.cs
--- a/Engine/Source/Programs/UnrealCloudDDC/Jupiter/Implementation/ContentId/MongoContentIdStore.cs
+++ b/Engine/Source/Programs/UnrealCloudDDC/Jupiter/Implementation/ContentId/MongoContentIdStore.cs
@@ -3,6 +3,7 @@
 using System;
 using System.Collections.Generic;
 using System.Linq;
+using System.Threading;
 using System.Threading.Tasks;
 using EpicGames.Horde.Storage;
 using Microsoft.Extensions.Options;
@@ -36,14 +37,14 @@ namespace Jupiter.Implementation
 			AddIndexFor<MongoContentIdModelV0>().CreateOne(indexModel);
 		}
 
-		public async Task<BlobId[]?> ResolveAsync(NamespaceId ns, ContentId contentId, bool mustBeContentId)
+		public async Task<BlobId[]?> ResolveAsync(NamespaceId ns, ContentId contentId, bool mustBeContentId, CancellationToken cancellationToken)
 		{
 			IMongoCollection<MongoContentIdModelV0> collection = GetCollection<MongoContentIdModelV0>();
 
 			IAsyncCursor<MongoContentIdModelV0> cursor =
-				await collection.FindAsync(model => model.Ns.Equals(ns.ToString(), StringComparison.OrdinalIgnoreCase) && model.ContentId.Equals(contentId.ToString(), StringComparison.Ordinal));
+				await collection.FindAsync(model => model.Ns.Equals(ns.ToString(), StringComparison.OrdinalIgnoreCase) && model.ContentId.Equals(contentId.ToString(), StringComparison.Ordinal), cancellationToken: cancellationToken);
 
-			MongoContentIdModelV0 model = await cursor.FirstOrDefaultAsync();
+			MongoContentIdModelV0 model = await cursor.FirstOrDefaultAsync(cancellationToken);
 			if (model != null)
 			{
 				foreach (int weight in model.ContentWeightToBlobsMap.Keys.OrderBy(contentWeight => contentWeight))
@@ -51,7 +52,7 @@ namespace Jupiter.Implementation
 					BlobId[] blobs = model.ContentWeightToBlobsMap[weight].Select(s => new BlobId(s)).ToArray();
 
 					{
-						BlobId[] missingBlobs = await _blobStore.FilterOutKnownBlobsAsync(ns, blobs);
+						BlobId[] missingBlobs = await _blobStore.FilterOutKnownBlobsAsync(ns, blobs, cancellationToken);
 						if (missingBlobs.Length == 0)
 						{
 							return blobs;
@@ -64,7 +65,7 @@ namespace Jupiter.Implementation
 
 			BlobId contentIdAsBlobIdentifier = contentId.AsBlobIdentifier();
 			// no content id where all blobs are present, check if its present in the blob store as a uncompressed version of the blob
-			if (!mustBeContentId && await _blobStore.ExistsAsync(ns, contentIdAsBlobIdentifier))
+			if (!mustBeContentId && await _blobStore.ExistsAsync(ns, contentIdAsBlobIdentifier, cancellationToken: cancellationToken))
 			{
 				return new[] { contentIdAsBlobIdentifier };
 			}
@@ -72,7 +73,7 @@ namespace Jupiter.Implementation
 			return null;
 		}
 
-		public async Task PutAsync(NamespaceId ns, ContentId contentId, BlobId blobIdentifier, int contentWeight)
+		public async Task PutAsync(NamespaceId ns, ContentId contentId, BlobId blobIdentifier, int contentWeight, CancellationToken cancellationToken)
 		{
 			IMongoCollection<MongoContentIdModelV0> collection = GetCollection<MongoContentIdModelV0>();
 
@@ -82,7 +83,7 @@ namespace Jupiter.Implementation
 			await collection.FindOneAndUpdateAsync(filter, update, new FindOneAndUpdateOptions<MongoContentIdModelV0>()
 			{
 				IsUpsert = true
-			});
+			}, cancellationToken);
 		}
 	}
 
diff --git a/Engine/Source/Programs/UnrealCloudDDC/Jupiter/Implementation/ContentId/ScyllaContentIdStore.cs b/Engine/Source/Programs/UnrealCloudDDC/Jupiter/Implementation/ContentId/ScyllaContentIdStore.cs
--- a/Engine/Source/Programs/UnrealCloudDDC/Jupiter/Implementation/ContentId/ScyllaContentIdStore.cs
+++ b/Engine/Source/Programs/UnrealCloudDDC/Jupiter/Implementation/ContentId/ScyllaContentIdStore.cs
@@ -2,6 +2,7 @@
 
 using System;
 using System.Linq;
+using System.Threading;
 using System.Threading.Tasks;
 using Cassandra;
 using Cassandra.Mapping;
@@ -41,7 +42,7 @@ namespace Jupiter.Implementation
 			}
 		}
 
-		public async Task<BlobId[]?> ResolveAsync(NamespaceId ns, ContentId contentId, bool mustBeContentId)
+		public async Task<BlobId[]?> ResolveAsync(NamespaceId ns, ContentId contentId, bool mustBeContentId, CancellationToken cancellationToken)
 		{
 			using TelemetrySpan scope = _tracer.BuildScyllaSpan("ScyllaContentIdStore.ResolveContentId").SetAttribute("resource.name", contentId.ToString());
 
@@ -65,7 +66,7 @@ namespace Jupiter.Implementation
 						{
 							using TelemetrySpan _ = _tracer.StartActiveSpan("ScyllaContentIdStore.FindMissingBlobs").SetAttribute("operation.name", "ScyllaContentIdStore.FindMissingBlobs");
 
-							BlobId[] missingBlobs = await _blobStore.FilterOutKnownBlobsAsync(ns, blobs);
+							BlobId[] missingBlobs = await _blobStore.FilterOutKnownBlobsAsync(ns, blobs, cancellationToken);
 							if (missingBlobs.Length == 0)
 							{
 								return blobs;
@@ -89,7 +90,7 @@ namespace Jupiter.Implementation
 						{
 							using TelemetrySpan _ = _tracer.StartActiveSpan("ScyllaContentIdStore.FindMissingBlobs").SetAttribute("operation.name", "ScyllaContentIdStore.FindMissingBlobs");
 
-							BlobId[] missingBlobs = await _blobStore.FilterOutKnownBlobsAsync(ns, blobs);
+							BlobId[] missingBlobs = await _blobStore.FilterOutKnownBlobsAsync(ns, blobs, cancellationToken);
 							if (missingBlobs.Length == 0)
 							{
 								return blobs;
@@ -103,7 +104,7 @@ namespace Jupiter.Implementation
 			if (!mustBeContentId)
 			{
 				// if no content id is found, but we have a blob that matches the content id (so a unchunked and uncompressed version of the data) we use that instead
-				bool contentIdBlobExists = await _blobStore.ExistsAsync(ns, contentIdBlob)!;
+				bool contentIdBlobExists = await _blobStore.ExistsAsync(ns, contentIdBlob, cancellationToken: cancellationToken)!;
 
 				if (contentIdBlobExists)
 				{
@@ -115,7 +116,7 @@ namespace Jupiter.Implementation
 			return null;
 		}
 
-		public async Task PutAsync(NamespaceId ns, ContentId contentId, BlobId blobIdentifier, int contentWeight)
+		public async Task PutAsync(NamespaceId ns, ContentId contentId, BlobId blobIdentifier, int contentWeight, CancellationToken cancellationToken)
 		{
 			using TelemetrySpan scope = _tracer.BuildScyllaSpan("ScyllaContentIdStore.PutContentId").SetAttribute("resource.name", contentId.ToString());
 			if (_scyllaSessionManager.IsScylla)
diff --git a/Engine/Source/Programs/UnrealCloudDDC/Jupiter/Implementation/GC/CacheCleanup.cs b/Engine/Source/Programs/UnrealCloudDDC/Jupiter/Implementation/GC/CacheCleanup.cs
--- a/Engine/Source/Programs/UnrealCloudDDC/Jupiter/Implementation/GC/CacheCleanup.cs
+++ b/Engine/Source/Programs/UnrealCloudDDC/Jupiter/Implementation/GC/CacheCleanup.cs
@@ -95,7 +95,7 @@ namespace Jupiter.Implementation
 			long consideredCount = 0;
 			DateTime cleanupStart = DateTime.Now;
 
-			await Parallel.ForEachAsync(_referencesStore.GetRecordsAsync(),
+			await Parallel.ForEachAsync(_referencesStore.GetRecordsAsync(cancellationToken),
 				new ParallelOptions
 				{
 					MaxDegreeOfParallelism = _settings.CurrentValue.OrphanRefMaxParallelOperations,
@@ -122,7 +122,7 @@ namespace Jupiter.Implementation
 							"Attempting to delete object {Namespace} {Bucket} {Name} as it was last updated {LastAccessTime} which is older then {CutoffTime}",
 							ns, bucket, name, lastAccessTime, cutoffTime);
 
-						await DeleteRefAsync(ns, bucket, name);
+						await DeleteRefAsync(ns, bucket, name, cancellationToken);
 
 						Interlocked.Increment(ref countOfDeletedRecords);
 
@@ -134,12 +134,12 @@ namespace Jupiter.Implementation
 					{
 						try
 						{
-							RefRecord refRecord = await _referencesStore.GetAsync(ns, bucket, name, IReferencesStore.FieldFlags.None, IReferencesStore.OperationFlags.BypassCache);
+							RefRecord refRecord = await _referencesStore.GetAsync(ns, bucket, name, IReferencesStore.FieldFlags.None, IReferencesStore.OperationFlags.BypassCache, cancellationToken);
 							if (!refRecord.IsFinalized)
 							{
 								_logger.LogInformation("Deleting object {Namespace} {Bucket} {Name} as it is not finalized. Was last accessed at {LastAccessTime}", ns, bucket, name, lastAccessTime);
 
-								await DeleteRefAsync(ns, bucket, name);
+								await DeleteRefAsync(ns, bucket, name, cancellationToken);
 								return;
 							}
 						}
@@ -157,7 +157,7 @@ namespace Jupiter.Implementation
 			return countOfDeletedRecords;
 		}
 
-		private async Task<bool> DeleteRefAsync(NamespaceId ns, BucketId bucket, RefId name)
+		private async Task<bool> DeleteRefAsync(NamespaceId ns, BucketId bucket, RefId name, CancellationToken cancellationToken)
 		{
 			using TelemetrySpan scope = _tracer.StartActiveSpan("gc.ref")
 				.SetAttribute("operation.name", "gc.ref")
@@ -173,12 +173,12 @@ namespace Jupiter.Implementation
 				{
 					bucketStatsCleanupTask = Task.Run(async () =>
 					{
-						List<BlobId> blobs = await _objectService.GetReferencedBlobsAsync(ns, bucket, name, ignoreMissingBlobs: true);
-						await _blobIndex.RemoveBlobFromBucketListAsync(ns, bucket, name, blobs);
-					});
+						List<BlobId> blobs = await _objectService.GetReferencedBlobsAsync(ns, bucket, name, ignoreMissingBlobs: true, cancellationToken: cancellationToken);
+						await _blobIndex.RemoveBlobFromBucketListAsync(ns, bucket, name, blobs, cancellationToken);
+					}, cancellationToken);
 				}
 
-				storeDelete = await _referencesStore.DeleteAsync(ns, bucket, name);
+				storeDelete = await _referencesStore.DeleteAsync(ns, bucket, name, cancellationToken);
 				if (storeDelete && _settings.CurrentValue.WriteDeleteToReplicationLog)
 				{
 					// insert a delete event into the transaction log
diff --git a/Engine/Source/Programs/UnrealCloudDDC/Jupiter/Implementation/GC/OrphanBlobCleanupRefs.cs b/Engine/Source/Programs/UnrealCloudDDC/Jupiter/Implementation/GC/OrphanBlobCleanupRefs.cs
--- a/Engine/Source/Programs/UnrealCloudDDC/Jupiter/Implementation/GC/OrphanBlobCleanupRefs.cs
+++ b/Engine/Source/Programs/UnrealCloudDDC/Jupiter/Implementation/GC/OrphanBlobCleanupRefs.cs
@@ -88,7 +88,7 @@ namespace Jupiter.Implementation
 				// only consider blobs that have been around for 60 minutes
 				// this due to cases were blobs are uploaded first
 				DateTime cutoff = DateTime.Now.AddMinutes(-60);
-				await Parallel.ForEachAsync(_blobService.ListObjectsAsync(@namespace),
+				await Parallel.ForEachAsync(_blobService.ListObjectsAsync(@namespace, cancellationToken),
 					new ParallelOptions { MaxDegreeOfParallelism = _gcSettings.CurrentValue.OrphanGCMaxParallelOperations, CancellationToken = cancellationToken },
 					async (tuple, ctx) =>
 					{
@@ -142,7 +142,7 @@ namespace Jupiter.Implementation
 					break;
 				}
 
-				IAsyncEnumerable<BaseBlobReference> references = _blobIndex.GetBlobReferencesAsync(blobNamespace, blob);
+				IAsyncEnumerable<BaseBlobReference> references = _blobIndex.GetBlobReferencesAsync(blobNamespace, blob, cancellationToken);
 
 				List<BaseBlobReference> oldReferences = new List<BaseBlobReference>();
 
@@ -160,7 +160,7 @@ namespace Jupiter.Implementation
 
 						try
 						{
-							found = await _refService.ExistsAsync(blobNamespace, bucket, key);
+							found = await _refService.ExistsAsync(blobNamespace, bucket, key, cancellationToken);
 							break;
 						}
 						catch (RefNotFoundException)
@@ -172,7 +172,7 @@ namespace Jupiter.Implementation
 					else if (baseBlobReference is BlobToBlobReference blobReference)
 					{
 						BlobId referringBlob = blobReference.Blob;
-						bool blobFound = await _blobService.ExistsAsync(blobNamespace, referringBlob);
+						bool blobFound = await _blobService.ExistsAsync(blobNamespace, referringBlob, cancellationToken: cancellationToken);
 						if (blobFound)
 						{
 							found = true;
@@ -192,7 +192,7 @@ namespace Jupiter.Implementation
 				if (found)
 				{
 					// if the object is still alive but had old references we remove the old references to keep the size of the references array more reasonable
-					await _blobIndex.RemoveReferencesAsync(blobNamespace, blob, oldReferences);
+					await _blobIndex.RemoveReferencesAsync(blobNamespace, blob, oldReferences, cancellationToken);
 				}
 			}
 
@@ -215,19 +215,19 @@ namespace Jupiter.Implementation
 			_logger.LogInformation("GC Orphan blob {Blob} from {StoragePool} which was last modified at {LastModifiedTime}", blob, storagePoolName, lastModifiedTime);
 
 			// if the blob was not found to have a reference in any of the namespace that share a storage pool then the blob is not used anymore and should be deleted from all the namespaces
-			await Parallel.ForEachAsync(namespacesThatSharePool, cancellationToken, async (ns, _) =>
+			await Parallel.ForEachAsync(namespacesThatSharePool, cancellationToken, async (ns, ctx) =>
 			{
-				await RemoveBlobAsync(ns, blob);
+				await RemoveBlobAsync(ns, blob, ctx);
 			});
 			return true;
 
 		}
 
-		private async Task RemoveBlobAsync(NamespaceId ns, BlobId blob)
+		private async Task RemoveBlobAsync(NamespaceId ns, BlobId blob, CancellationToken cancellationToken)
 		{
 			try
 			{
-				await _blobService.DeleteObjectAsync(ns, blob);
+				await _blobService.DeleteObjectAsync(ns, blob, cancellationToken);
 			}
 			catch (BlobNotFoundException)
 			{
@@ -241,7 +241,7 @@ namespace Jupiter.Implementation
 
 		private IAsyncEnumerable<NamespaceId> ListNamespaces()
 		{
-			return _refService.GetNamespacesAsync();
+			return _refService.GetNamespacesAsync(CancellationToken.None);
 		}
 	}
 }
diff --git a/Engine/Source/Programs/UnrealCloudDDC/Jupiter/Implementation/IRefService.cs b/Engine/Source/Programs/UnrealCloudDDC/Jupiter/Implementation/IRefService.cs
--- a/Engine/Source/Programs/UnrealCloudDDC/Jupiter/Implementation/IRefService.cs
+++ b/Engine/Source/Programs/UnrealCloudDDC/Jupiter/Implementation/IRefService.cs
@@ -2,6 +2,7 @@
 
 using System;
 using System.Collections.Generic;
+using System.Threading;
 using System.Threading.Tasks;
 using EpicGames.Horde.Storage;
 using EpicGames.Serialization;
@@ -10,18 +11,18 @@ namespace Jupiter.Implementation
 {
 	public interface IRefService
 	{
-		Task<(RefRecord, BlobContents?)> GetAsync(NamespaceId ns, BucketId bucket, RefId key, string[] fields, bool doLastAccessTracking = true);
-		Task<(ContentId[], BlobId[])> PutAsync(NamespaceId ns, BucketId bucket, RefId key, BlobId blobHash, CbObject payload);
-		Task<(ContentId[], BlobId[])> FinalizeAsync(NamespaceId ns, BucketId bucket, RefId key, BlobId blobHash);
+		Task<(RefRecord, BlobContents?)> GetAsync(NamespaceId ns, BucketId bucket, RefId key, string[] fields, bool doLastAccessTracking = true, CancellationToken cancellationToken = default);
+		Task<(ContentId[], BlobId[])> PutAsync(NamespaceId ns, BucketId bucket, RefId key, BlobId blobHash, CbObject payload, CancellationToken cancellationToken = default);
+		Task<(ContentId[], BlobId[])> FinalizeAsync(NamespaceId ns, BucketId bucket, RefId key, BlobId blobHash, CancellationToken cancellationToken = default);
 
-		IAsyncEnumerable<NamespaceId> GetNamespacesAsync();
+		IAsyncEnumerable<NamespaceId> GetNamespacesAsync(CancellationToken cancellationToken = default);
 
-		Task<bool> DeleteAsync(NamespaceId ns, BucketId bucket, RefId key);
-		Task<long> DropNamespaceAsync(NamespaceId ns);
-		Task<long> DeleteBucketAsync(NamespaceId ns, BucketId bucket);
+		Task<bool> DeleteAsync(NamespaceId ns, BucketId bucket, RefId key, CancellationToken cancellationToken = default);
+		Task<long> DropNamespaceAsync(NamespaceId ns, CancellationToken cancellationToken = default);
+		Task<long> DeleteBucketAsync(NamespaceId ns, BucketId bucket, CancellationToken cancellationToken = default);
 
-		Task<bool> ExistsAsync(NamespaceId ns, BucketId bucket, RefId key);
-		Task<List<BlobId>> GetReferencedBlobsAsync(NamespaceId ns, BucketId bucket, RefId key, bool ignoreMissingBlobs);
+		Task<bool> ExistsAsync(NamespaceId ns, BucketId bucket, RefId key, CancellationToken cancellationToken = default);
+		Task<List<BlobId>> GetReferencedBlobsAsync(NamespaceId ns, BucketId bucket, RefId key, bool ignoreMissingBlobs, CancellationToken cancellationToken = default);
 	}
 
 	public class ObjectHashMismatchException : Exception
diff --git a/Engine/Source/Programs/UnrealCloudDDC/Jupiter/Implementation/IReferenceResolver.cs b/Engine/Source/Programs/UnrealCloudDDC/Jupiter/Implementation/IReferenceResolver.cs
--- a/Engine/Source/Programs/UnrealCloudDDC/Jupiter/Implementation/IReferenceResolver.cs
+++ b/Engine/Source/Programs/UnrealCloudDDC/Jupiter/Implementation/IReferenceResolver.cs
@@ -3,6 +3,8 @@
 using System;
 using System.Collections.Generic;
 using System.Linq;
+using System.Runtime.CompilerServices;
+using System.Threading;
 using System.Threading.Tasks;
 using EpicGames.Core;
 using EpicGames.Horde.Storage;
@@ -74,8 +76,9 @@ namespace Jupiter.Implementation
 		/// <param name="ns">The namespace to check</param>
 		/// <param name="cb">The compact binary object to resolve references for</param>
 		/// <param name="ignoreMissingBlobs">Set to true to always returned the blobs found, ignoring anything that is missing rather then throwing</param>
+		/// <param name="cancellationToken">Cancellation token for the operation</param>
 		/// <returns></returns>
-		IAsyncEnumerable<BlobId> GetReferencedBlobsAsync(NamespaceId ns, CbObject cb, bool ignoreMissingBlobs = false);
+		IAsyncEnumerable<BlobId> GetReferencedBlobsAsync(NamespaceId ns, CbObject cb, bool ignoreMissingBlobs = false, CancellationToken cancellationToken = default);
 
 		/// <summary>
 		/// Returns which attachments exist in the cb object or any children
@@ -83,8 +86,9 @@ namespace Jupiter.Implementation
 		/// </summary>
 		/// <param name="ns">The namespace to check</param>
 		/// <param name="cb">The compact binary object to resolve references for</param>
+		/// <param name="cancellationToken">Cancellation token for the operation</param>
 		/// <returns></returns>
-		IAsyncEnumerable<Attachment> GetAttachmentsAsync(NamespaceId ns, CbObject cb);
+		IAsyncEnumerable<Attachment> GetAttachmentsAsync(NamespaceId ns, CbObject cb, CancellationToken cancellationToken);
 	}
 
 	public class ReferenceResolver : IReferenceResolver
@@ -100,7 +104,7 @@ namespace Jupiter.Implementation
 			_tracer = tracer;
 		}
 
-		public async IAsyncEnumerable<Attachment> GetAttachmentsAsync(NamespaceId ns, CbObject cb)
+		public async IAsyncEnumerable<Attachment> GetAttachmentsAsync(NamespaceId ns, CbObject cb, [EnumeratorCancellation] CancellationToken cancellationToken)
 		{
 			Queue<CbObject> objectsToVisit = new Queue<CbObject>();
 			objectsToVisit.Enqueue(cb);
@@ -124,13 +128,13 @@ namespace Jupiter.Implementation
 
 						if (field.IsBinaryAttachment())
 						{
-							Task<(ContentId, BlobId[]?)> resolveContentId = ResolveContentIdAsync(ns, contentId);
+							Task<(ContentId, BlobId[]?)> resolveContentId = ResolveContentIdAsync(ns, contentId, cancellationToken);
 							pendingContentIdResolves.Add(resolveContentId);
 						}
 						else if (field.IsObjectAttachment())
 						{
 							attachments.Add(new ObjectAttachment(blobIdentifier));
-							pendingCompactBinaryAttachments.Add(ParseCompactBinaryAttachmentAsync(ns, blobIdentifier));
+							pendingCompactBinaryAttachments.Add(ParseCompactBinaryAttachmentAsync(ns, blobIdentifier, cancellationToken));
 						}
 						else
 						{
@@ -234,7 +238,7 @@ namespace Jupiter.Implementation
 			}
 		}
 
-		public async IAsyncEnumerable<BlobId> GetReferencedBlobsAsync(NamespaceId ns, CbObject cb, bool ignoreMissingBlobs = false)
+		public async IAsyncEnumerable<BlobId> GetReferencedBlobsAsync(NamespaceId ns, CbObject cb, bool ignoreMissingBlobs = false, [EnumeratorCancellation] CancellationToken cancellationToken = default)
 		{
 			List<Task<(BlobId, bool)>> pendingBlobExistsChecks = new();
 			List<Task<(ContentIdAttachment, bool)>> pendingContentIdChecks = new();
@@ -242,21 +246,21 @@ namespace Jupiter.Implementation
 			List<BlobId> unresolvedBlobReferences = new List<BlobId>();
 
 			// Resolve all the attachments
-			await foreach (Attachment attachment in GetAttachmentsAsync(ns, cb))
+			await foreach (Attachment attachment in GetAttachmentsAsync(ns, cb, cancellationToken))
 			{
 				if (attachment is BlobAttachment blobAttachment)
 				{
-					pendingBlobExistsChecks.Add(CheckBlobExistsAsync(ns, blobAttachment.Identifier));
+					pendingBlobExistsChecks.Add(CheckBlobExistsAsync(ns, blobAttachment.Identifier, cancellationToken));
 				}
 				else if (attachment is ContentIdAttachment contentIdAttachment)
 				{
 					// If we find a content id we resolve that into the actual blobs it references
-					pendingContentIdChecks.Add(CheckContentIdExistsAsync(ns, contentIdAttachment));
+					pendingContentIdChecks.Add(CheckContentIdExistsAsync(ns, contentIdAttachment, cancellationToken));
 				}
 				else if (attachment is ObjectAttachment objectAttachment)
 				{
 					// a object just references the same blob, traversing the object attachment is done in GetAttachments
-					pendingBlobExistsChecks.Add(CheckBlobExistsAsync(ns, objectAttachment.Identifier));
+					pendingBlobExistsChecks.Add(CheckBlobExistsAsync(ns, objectAttachment.Identifier, cancellationToken));
 				}
 				else
 				{
@@ -306,12 +310,12 @@ namespace Jupiter.Implementation
 			}
 		}
 
-		private async Task<(ContentIdAttachment, bool)> CheckContentIdExistsAsync(NamespaceId ns, ContentIdAttachment contentIdAttachment)
+		private async Task<(ContentIdAttachment, bool)> CheckContentIdExistsAsync(NamespaceId ns, ContentIdAttachment contentIdAttachment, CancellationToken cancellationToken)
 		{
 			bool allBlobsExist = true;
 			foreach (BlobId b in contentIdAttachment.ReferencedBlobs)
 			{
-				(BlobId _, bool exists) = await CheckBlobExistsAsync(ns, b);
+				(BlobId _, bool exists) = await CheckBlobExistsAsync(ns, b, cancellationToken);
 
 				if (!exists)
 				{
@@ -323,27 +327,27 @@ namespace Jupiter.Implementation
 			return (contentIdAttachment, allBlobsExist);
 		}
 
-		private async Task<CbObject> ParseCompactBinaryAttachmentAsync(NamespaceId ns, BlobId blobIdentifier)
+		private async Task<CbObject> ParseCompactBinaryAttachmentAsync(NamespaceId ns, BlobId blobIdentifier, CancellationToken cancellationToken)
 		{
-			BlobContents contents = await _blobStore.GetObjectAsync(ns, blobIdentifier);
-			byte[] data = await contents.Stream.ToByteArrayAsync();
+			BlobContents contents = await _blobStore.GetObjectAsync(ns, blobIdentifier, cancellationToken: cancellationToken);
+			byte[] data = await contents.Stream.ToByteArrayAsync(cancellationToken);
 			CbObject childBinaryObject = new CbObject(data);
 
 			return childBinaryObject;
 		}
 
-		private async Task<(ContentId, BlobId[]?)> ResolveContentIdAsync(NamespaceId ns, ContentId contentId)
+		private async Task<(ContentId, BlobId[]?)> ResolveContentIdAsync(NamespaceId ns, ContentId contentId, CancellationToken cancellationToken)
 		{
 			using TelemetrySpan scope = _tracer.StartActiveSpan("ReferenceResolver.ResolveContentId")
 				.SetAttribute("operation.name", "ReferenceResolver.ResolveContentId")
 				.SetAttribute("resource.name", contentId.ToString());
-			BlobId[]? resolvedBlobs = await _contentIdStore.ResolveAsync(ns, contentId);
+			BlobId[]? resolvedBlobs = await _contentIdStore.ResolveAsync(ns, contentId, cancellationToken: cancellationToken);
 			return (contentId, resolvedBlobs);
 		}
 
-		private async Task<(BlobId, bool)> CheckBlobExistsAsync(NamespaceId ns, BlobId blob)
+		private async Task<(BlobId, bool)> CheckBlobExistsAsync(NamespaceId ns, BlobId blob, CancellationToken cancellationToken)
 		{
-			return (blob, await _blobStore.ExistsAsync(ns, blob));
+			return (blob, await _blobStore.ExistsAsync(ns, blob, cancellationToken: cancellationToken));
 		}
 	}
 
diff --git a/Engine/Source/Programs/UnrealCloudDDC/Jupiter/Implementation/LastAccess/LastAccessServiceReferences.cs b/Engine/Source/Programs/UnrealCloudDDC/Jupiter/Implementation/LastAccess/LastAccessServiceReferences.cs
--- a/Engine/Source/Programs/UnrealCloudDDC/Jupiter/Implementation/LastAccess/LastAccessServiceReferences.cs
+++ b/Engine/Source/Programs/UnrealCloudDDC/Jupiter/Implementation/LastAccess/LastAccessServiceReferences.cs
@@ -55,7 +55,7 @@ namespace Jupiter.Implementation
 			Running = false;
 
 			// process the last records we have built up
-			await ProcessLastAccessRecordsAsync();
+			await ProcessLastAccessRecordsAsync(cancellationToken);
 		}
 
 		private void OnUpdate(object? state)
@@ -63,7 +63,7 @@ namespace Jupiter.Implementation
 			try
 			{
 				// call results to make sure we join the task
-				ProcessLastAccessRecordsAsync().Wait();
+				ProcessLastAccessRecordsAsync(CancellationToken.None).Wait();
 			}
 			catch (Exception e)
 			{
@@ -73,7 +73,7 @@ namespace Jupiter.Implementation
 			}
 		}
 
-		internal async Task<List<(LastAccessRecord, DateTime)>> ProcessLastAccessRecordsAsync()
+		internal async Task<List<(LastAccessRecord, DateTime)>> ProcessLastAccessRecordsAsync(CancellationToken cancellationToken)
 		{
 			_logger.LogInformation("Running Last Access Aggregation for refs");
 			List<(LastAccessRecord, DateTime)> records = await _lastAccessCacheRecord.GetLastAccessedRecords();
@@ -88,9 +88,9 @@ namespace Jupiter.Implementation
 					.SetAttribute("operation.name", "lastAccess.update")
 					.SetAttribute("resource.name", $"{record.Namespace}:{record.Bucket}.{record.Key}");
 				_logger.LogDebug("Updating last access time to {LastAccessTime} for {Record}", lastAccessTime, record);
-				await _referencesStore.UpdateLastAccessTimeAsync(record.Namespace, record.Bucket, record.Key, lastAccessTime);
+				await _referencesStore.UpdateLastAccessTimeAsync(record.Namespace, record.Bucket, record.Key, lastAccessTime, cancellationToken);
 				// delay 10ms between each record to distribute the load more evenly for the db
-				await Task.Delay(10);
+				await Task.Delay(10, cancellationToken);
 			}
 
 			return records;
diff --git a/Engine/Source/Programs/UnrealCloudDDC/Jupiter/Implementation/MetricsService.cs b/Engine/Source/Programs/UnrealCloudDDC/Jupiter/Implementation/MetricsService.cs
--- a/Engine/Source/Programs/UnrealCloudDDC/Jupiter/Implementation/MetricsService.cs
+++ b/Engine/Source/Programs/UnrealCloudDDC/Jupiter/Implementation/MetricsService.cs
@@ -88,14 +88,14 @@ namespace Jupiter.Implementation
 			_logger.LogInformation("Attempting to calculate metrics. ");
 			try
 			{
-				await foreach (NamespaceId ns in _referencesStore.GetNamespacesAsync().WithCancellation(cancellationToken))
+				await foreach (NamespaceId ns in _referencesStore.GetNamespacesAsync(cancellationToken).WithCancellation(cancellationToken))
 				{
-					await foreach (BucketId bucket in _referencesStore.GetBucketsAsync(ns).WithCancellation(cancellationToken))
+					await foreach (BucketId bucket in _referencesStore.GetBucketsAsync(ns, cancellationToken).WithCancellation(cancellationToken))
 					{
 						DateTime start = DateTime.UtcNow;
 						_logger.LogInformation("Calculating stats for {Namespace} {Bucket}", ns, bucket);
 
-						await calculator.CalculateStatsForBucketAsync(ns, bucket);
+						await calculator.CalculateStatsForBucketAsync(ns, bucket, cancellationToken);
 
 						TimeSpan duration = DateTime.UtcNow - start;
 						_logger.LogInformation("Stats calculated for {Namespace} {Bucket} took {Duration}", ns, bucket, duration);
@@ -136,7 +136,7 @@ namespace Jupiter.Implementation
 			_refsInBucketGauge = meter.CreateGauge<long>("blobstats.refs_in_bucket");
 		}
 
-		public async Task<BucketStats?> CalculateStatsForBucketAsync(NamespaceId ns, BucketId bucket)
+		public async Task<BucketStats?> CalculateStatsForBucketAsync(NamespaceId ns, BucketId bucket, CancellationToken cancellationToken = default)
 		{
 			using TelemetrySpan removeBlobScope = _tracer.StartActiveSpan("metrics.calculate")
 				.SetAttribute("operation.name", "metrics.calculate")
@@ -144,7 +144,7 @@ namespace Jupiter.Implementation
 
 			KeyValuePair<string, object?>[] tags = new[] { new KeyValuePair<string, object?>("Bucket", bucket.ToString()), new KeyValuePair<string, object?>("Namespace", ns.ToString()) };
 
-			BucketStats stats = await _blobIndex.CalculateBucketStatisticsAsync(ns, bucket);
+			BucketStats stats = await _blobIndex.CalculateBucketStatisticsAsync(ns, bucket, cancellationToken);
 
 			_blobSizeAvgGauge.Record(stats.AvgSize, tags);
 			_blobSizeMinGauge.Record(stats.SmallestBlobFound, tags);
diff --git a/Engine/Source/Programs/UnrealCloudDDC/Jupiter/Implementation/Objects/CacheReferencesStore.cs b/Engine/Source/Programs/UnrealCloudDDC/Jupiter/Implementation/Objects/CacheReferencesStore.cs
--- a/Engine/Source/Programs/UnrealCloudDDC/Jupiter/Implementation/Objects/CacheReferencesStore.cs
+++ b/Engine/Source/Programs/UnrealCloudDDC/Jupiter/Implementation/Objects/CacheReferencesStore.cs
@@ -8,6 +8,7 @@ using System.Net.Http;
 using System.Net.Http.Headers;
 using System.Net.Http.Json;
 using System.Net.Mime;
+using System.Threading;
 using System.Threading.Tasks;
 using EpicGames.AspNet;
 using EpicGames.Horde.Storage;
@@ -31,11 +32,11 @@ namespace Jupiter.Implementation
 			_upstreamReferenceStore = ActivatorUtilities.CreateInstance<UpstreamReferenceStore>(provider);
 		}
 
-		public async Task<RefRecord> GetAsync(NamespaceId ns, BucketId bucket, RefId key, IReferencesStore.FieldFlags fieldFlags, IReferencesStore.OperationFlags opFlags)
+		public async Task<RefRecord> GetAsync(NamespaceId ns, BucketId bucket, RefId key, IReferencesStore.FieldFlags fieldFlags, IReferencesStore.OperationFlags opFlags, CancellationToken cancellationToken)
 		{
 			try
 			{
-				RefRecord cachedRecord = await _mongoReferenceStore.GetAsync(ns, bucket, key, fieldFlags, opFlags);
+				RefRecord cachedRecord = await _mongoReferenceStore.GetAsync(ns, bucket, key, fieldFlags, opFlags, cancellationToken);
 				return cachedRecord;
 			}
 			catch (RefNotFoundException)
@@ -43,92 +44,92 @@ namespace Jupiter.Implementation
 				// not cached, we check the upstream for it
 			}
 
-			RefRecord record = await _upstreamReferenceStore.GetAsync(ns, bucket, key, fieldFlags, opFlags);
-			await _mongoReferenceStore.PutAsync(record.Namespace, record.Bucket, record.Name, record.BlobIdentifier, record.InlinePayload, record.IsFinalized);
+			RefRecord record = await _upstreamReferenceStore.GetAsync(ns, bucket, key, fieldFlags, opFlags, cancellationToken);
+			await _mongoReferenceStore.PutAsync(record.Namespace, record.Bucket, record.Name, record.BlobIdentifier, record.InlinePayload, record.IsFinalized, cancellationToken);
 			return record;
 		}
 
-		public async Task PutAsync(NamespaceId ns, BucketId bucket, RefId key, BlobId blobHash, byte[] blob, bool isFinalized)
+		public async Task PutAsync(NamespaceId ns, BucketId bucket, RefId key, BlobId blobHash, byte[] blob, bool isFinalized, CancellationToken cancellationToken)
 		{
-			Task cachePut = _mongoReferenceStore.PutAsync(ns, bucket, key, blobHash, blob, isFinalized);
-			Task upstreamPut = _upstreamReferenceStore.PutAsync(ns, bucket, key, blobHash, blob, isFinalized);
+			Task cachePut = _mongoReferenceStore.PutAsync(ns, bucket, key, blobHash, blob, isFinalized, cancellationToken);
+			Task upstreamPut = _upstreamReferenceStore.PutAsync(ns, bucket, key, blobHash, blob, isFinalized, cancellationToken);
 
 			await Task.WhenAll(cachePut, upstreamPut);
 		}
 
-		public async Task FinalizeAsync(NamespaceId ns, BucketId bucket, RefId key, BlobId blobIdentifier)
+		public async Task FinalizeAsync(NamespaceId ns, BucketId bucket, RefId key, BlobId blobIdentifier, CancellationToken cancellationToken)
 		{
-			Task cacheFinalize = _mongoReferenceStore.FinalizeAsync(ns, bucket, key, blobIdentifier);
-			Task upstreamFinalize = _upstreamReferenceStore.FinalizeAsync(ns, bucket, key, blobIdentifier);
+			Task cacheFinalize = _mongoReferenceStore.FinalizeAsync(ns, bucket, key, blobIdentifier, cancellationToken);
+			Task upstreamFinalize = _upstreamReferenceStore.FinalizeAsync(ns, bucket, key, blobIdentifier, cancellationToken);
 
 			await Task.WhenAll(cacheFinalize, upstreamFinalize);
 		}
 
-		public async Task<DateTime?> GetLastAccessTimeAsync(NamespaceId ns, BucketId bucket, RefId key)
+		public async Task<DateTime?> GetLastAccessTimeAsync(NamespaceId ns, BucketId bucket, RefId key, CancellationToken cancellationToken)
 		{
 			try
 			{
-				DateTime? lastAccessTimeCache = await _mongoReferenceStore.GetLastAccessTimeAsync(ns, bucket, key);
+				DateTime? lastAccessTimeCache = await _mongoReferenceStore.GetLastAccessTimeAsync(ns, bucket, key, cancellationToken);
 				return lastAccessTimeCache;
 			}
 			catch (RefNotFoundException)
 			{
 				// not cached, we check the upstream for it
 			}
-			DateTime? lastAccessTime = await _upstreamReferenceStore.GetLastAccessTimeAsync(ns, bucket, key);
+			DateTime? lastAccessTime = await _upstreamReferenceStore.GetLastAccessTimeAsync(ns, bucket, key, cancellationToken);
 			if (lastAccessTime.HasValue)
 			{
-				await _mongoReferenceStore.UpdateLastAccessTimeAsync(ns, bucket, key, lastAccessTime.Value);
+				await _mongoReferenceStore.UpdateLastAccessTimeAsync(ns, bucket, key, lastAccessTime.Value, cancellationToken);
 			}
 
 			return lastAccessTime;
 		}
 
-		public async Task UpdateLastAccessTimeAsync(NamespaceId ns, BucketId bucket, RefId key, DateTime newLastAccessTime)
+		public async Task UpdateLastAccessTimeAsync(NamespaceId ns, BucketId bucket, RefId key, DateTime newLastAccessTime, CancellationToken cancellationToken)
 		{
-			Task cacheUpdateLastAccess = _mongoReferenceStore.UpdateLastAccessTimeAsync(ns, bucket, key, newLastAccessTime);
-			Task upstreamUpdateLastAccess = _upstreamReferenceStore.UpdateLastAccessTimeAsync(ns, bucket, key, newLastAccessTime);
+			Task cacheUpdateLastAccess = _mongoReferenceStore.UpdateLastAccessTimeAsync(ns, bucket, key, newLastAccessTime, cancellationToken);
+			Task upstreamUpdateLastAccess = _upstreamReferenceStore.UpdateLastAccessTimeAsync(ns, bucket, key, newLastAccessTime, cancellationToken);
 
 			await Task.WhenAll(cacheUpdateLastAccess, upstreamUpdateLastAccess);
 		}
 
-		public IAsyncEnumerable<(NamespaceId, BucketId, RefId, DateTime)> GetRecordsAsync()
+		public IAsyncEnumerable<(NamespaceId, BucketId, RefId, DateTime)> GetRecordsAsync(CancellationToken cancellationToken)
 		{
 			throw new NotImplementedException("GetRecords not supported on a cached reference store");
 		}
 
-		public IAsyncEnumerable<(NamespaceId, BucketId, RefId)> GetRecordsWithoutAccessTimeAsync()
+		public IAsyncEnumerable<(NamespaceId, BucketId, RefId)> GetRecordsWithoutAccessTimeAsync(CancellationToken cancellationToken)
 		{
 			throw new NotImplementedException("GetRecordsWithoutAccessTimeAsync not supported on a cached reference store");
 		}
 
-		public IAsyncEnumerable<(RefId, BlobId)> GetRecordsInBucketAsync(NamespaceId ns, BucketId bucket)
+		public IAsyncEnumerable<(RefId, BlobId)> GetRecordsInBucketAsync(NamespaceId ns, BucketId bucket, CancellationToken cancellationToken)
 		{
 			throw new NotImplementedException("GetRecordsInBucketAsync not supported on a cached reference store");
 		}
 
-		public IAsyncEnumerable<NamespaceId> GetNamespacesAsync()
+		public IAsyncEnumerable<NamespaceId> GetNamespacesAsync(CancellationToken cancellationToken)
 		{
 			throw new NotImplementedException("GetNamespaces not supported on a cached reference store");
 		}
 
-		public IAsyncEnumerable<BucketId> GetBucketsAsync(NamespaceId ns)
+		public IAsyncEnumerable<BucketId> GetBucketsAsync(NamespaceId ns, CancellationToken cancellationToken)
 		{
 			throw new NotImplementedException("GetBuckets not supported on a cached reference store");
 		}
 
-		public Task<bool> DeleteAsync(NamespaceId ns, BucketId bucket, RefId key)
+		public Task<bool> DeleteAsync(NamespaceId ns, BucketId bucket, RefId key, CancellationToken cancellationToken)
 		{
 			throw new NotImplementedException("Deletes are not supported on a cached reference store");
 		}
 
-		public Task<long> DropNamespaceAsync(NamespaceId ns)
+		public Task<long> DropNamespaceAsync(NamespaceId ns, CancellationToken cancellationToken)
 		{
 			throw new NotImplementedException("DropNamespace is not supported on a cached reference store");
 
 		}
 
-		public Task<long> DeleteBucketAsync(NamespaceId ns, BucketId bucket)
+		public Task<long> DeleteBucketAsync(NamespaceId ns, BucketId bucket, CancellationToken cancellationToken)
 		{
 			throw new NotImplementedException("DeleteBucket is not supported on a cached reference store");
 		}
@@ -140,11 +141,11 @@ namespace Jupiter.Implementation
 		{
 		}
 
-		public async Task<RefRecord> GetAsync(NamespaceId ns, BucketId bucket, RefId key, IReferencesStore.FieldFlags fieldFlags, IReferencesStore.OperationFlags opFlags)
+		public async Task<RefRecord> GetAsync(NamespaceId ns, BucketId bucket, RefId key, IReferencesStore.FieldFlags fieldFlags, IReferencesStore.OperationFlags opFlags, CancellationToken cancellationToken)
 		{
 			using HttpRequestMessage getObjectRequest = await BuildHttpRequestAsync(HttpMethod.Get, new Uri($"api/v1/refs/{ns}/{bucket}/{key}/metadata", UriKind.Relative));
 			getObjectRequest.Headers.Add("Accept", MediaTypeNames.Application.Json);
-			HttpResponseMessage response = await HttpClient.SendAsync(getObjectRequest);
+			HttpResponseMessage response = await HttpClient.SendAsync(getObjectRequest, cancellationToken);
 			if (response.StatusCode == HttpStatusCode.NotFound)
 			{
 				throw new RefNotFoundException(ns, bucket, key);
@@ -152,7 +153,7 @@ namespace Jupiter.Implementation
 
 			response.EnsureSuccessStatusCode();
 
-			RefMetadataResponse? metadataResponse = await response.Content.ReadFromJsonAsync<RefMetadataResponse>();
+			RefMetadataResponse? metadataResponse = await response.Content.ReadFromJsonAsync<RefMetadataResponse>(cancellationToken);
 			if (metadataResponse == null)
 			{
 				throw new Exception("Unable to deserialize metadata response");
@@ -160,7 +161,7 @@ namespace Jupiter.Implementation
 			return new RefRecord(metadataResponse.Ns, metadataResponse.Bucket, metadataResponse.Name, metadataResponse.LastAccess, metadataResponse.InlinePayload, metadataResponse.PayloadIdentifier, metadataResponse.IsFinalized);
 		}
 
-		public async Task PutAsync(NamespaceId ns, BucketId bucket, RefId key, BlobId blobHash, byte[] blob, bool isFinalized)
+		public async Task PutAsync(NamespaceId ns, BucketId bucket, RefId key, BlobId blobHash, byte[] blob, bool isFinalized, CancellationToken cancellationToken)
 		{
 			using HttpRequestMessage putObjectRequest = await BuildHttpRequestAsync(HttpMethod.Put, new Uri($"api/v1/refs/{ns}/{bucket}/{key}", UriKind.Relative));
 			putObjectRequest.Headers.Add("Accept", MediaTypeNames.Application.Json);
@@ -169,10 +170,10 @@ namespace Jupiter.Implementation
 			// a blob sent to the reference store is always by definition a compact binary, so we upload it as such to avoid any automatic conversions that we would get with octet-stream
 			putObjectRequest.Content.Headers.ContentType = new MediaTypeHeaderValue(CustomMediaTypeNames.UnrealCompactBinary);
 
-			HttpResponseMessage response = await HttpClient.SendAsync(putObjectRequest);
+			HttpResponseMessage response = await HttpClient.SendAsync(putObjectRequest, cancellationToken);
 			if (response.StatusCode == HttpStatusCode.BadRequest)
 			{
-				ProblemDetails? problem = await response.Content.ReadFromJsonAsync<ProblemDetails>();
+				ProblemDetails? problem = await response.Content.ReadFromJsonAsync<ProblemDetails>(cancellationToken);
 				if (problem == null)
 				{
 					throw new Exception("Unable to deserialize problem details when bad request was encountered");
@@ -184,16 +185,16 @@ namespace Jupiter.Implementation
 			// if this put returns needs, we cant really do anything about it. we should be calling put in the upstream blob store as the operation continues and we should be able to catch the missing blobs during the finalize call
 		}
 
-		public async Task FinalizeAsync(NamespaceId ns, BucketId bucket, RefId key, BlobId blobIdentifier)
+		public async Task FinalizeAsync(NamespaceId ns, BucketId bucket, RefId key, BlobId blobIdentifier, CancellationToken cancellationToken)
 		{
 			using HttpRequestMessage putObjectRequest = await BuildHttpRequestAsync(HttpMethod.Post, new Uri($"api/v1/refs/{ns}/{bucket}/{key}/finalize/{blobIdentifier}", UriKind.Relative));
 			putObjectRequest.Headers.Add("Accept", MediaTypeNames.Application.Json);
 
-			HttpResponseMessage response = await HttpClient.SendAsync(putObjectRequest);
+			HttpResponseMessage response = await HttpClient.SendAsync(putObjectRequest, cancellationToken);
 
 			if (response.StatusCode == HttpStatusCode.BadRequest)
 			{
-				ProblemDetails? problem = await response.Content.ReadFromJsonAsync<ProblemDetails>();
+				ProblemDetails? problem = await response.Content.ReadFromJsonAsync<ProblemDetails>(cancellationToken);
 				if (problem == null)
 				{
 					throw new Exception("Unable to deserialize problem details when bad request was encountered during finalize");
@@ -202,7 +203,7 @@ namespace Jupiter.Implementation
 			}
 			response.EnsureSuccessStatusCode();
 
-			PutObjectResponse? putObjectResponse = await response.Content.ReadFromJsonAsync<PutObjectResponse>();
+			PutObjectResponse? putObjectResponse = await response.Content.ReadFromJsonAsync<PutObjectResponse>(cancellationToken);
 			if (putObjectResponse == null)
 			{
 				throw new Exception("Unable to deserialize put object response when finalizing ref");
@@ -215,11 +216,11 @@ namespace Jupiter.Implementation
 			}
 		}
 
-		public async Task<DateTime?> GetLastAccessTimeAsync(NamespaceId ns, BucketId bucket, RefId key)
+		public async Task<DateTime?> GetLastAccessTimeAsync(NamespaceId ns, BucketId bucket, RefId key, CancellationToken cancellationToken)
 		{
 			try
 			{
-				RefRecord foo = await GetAsync(ns, bucket, key, IReferencesStore.FieldFlags.None, IReferencesStore.OperationFlags.None);
+				RefRecord foo = await GetAsync(ns, bucket, key, IReferencesStore.FieldFlags.None, IReferencesStore.OperationFlags.None, cancellationToken);
 				return foo.LastAccess;
 			}
 			catch (RefNotFoundException)
@@ -228,48 +229,48 @@ namespace Jupiter.Implementation
 			}
 		}
 
-		public async Task UpdateLastAccessTimeAsync(NamespaceId ns, BucketId bucket, RefId key, DateTime newLastAccessTime)
+		public async Task UpdateLastAccessTimeAsync(NamespaceId ns, BucketId bucket, RefId key, DateTime newLastAccessTime, CancellationToken cancellationToken)
 		{
 			// there is no endpoint to update last access time externally, but fetching a object will in turn update its last access time, though to a different time then newLastAccessTime
-			await GetAsync(ns, bucket, key, IReferencesStore.FieldFlags.None, IReferencesStore.OperationFlags.None);
+			await GetAsync(ns, bucket, key, IReferencesStore.FieldFlags.None, IReferencesStore.OperationFlags.None, cancellationToken);
 		}
 
-		public IAsyncEnumerable<(NamespaceId, BucketId, RefId, DateTime)> GetRecordsAsync()
+		public IAsyncEnumerable<(NamespaceId, BucketId, RefId, DateTime)> GetRecordsAsync(CancellationToken cancellationToken)
 		{
 			throw new NotImplementedException("GetRecords is not supported on a upstream reference store");
 		}
 
-		public IAsyncEnumerable<(NamespaceId, BucketId, RefId)> GetRecordsWithoutAccessTimeAsync()
+		public IAsyncEnumerable<(NamespaceId, BucketId, RefId)> GetRecordsWithoutAccessTimeAsync(CancellationToken cancellationToken)
 		{
 			throw new NotImplementedException("GetRecordsWithoutAccessTimeAsync is not supported on a upstream reference store");
 		}
 
-		public IAsyncEnumerable<(RefId, BlobId)> GetRecordsInBucketAsync(NamespaceId ns, BucketId bucket)
+		public IAsyncEnumerable<(RefId, BlobId)> GetRecordsInBucketAsync(NamespaceId ns, BucketId bucket, CancellationToken cancellationToken)
 		{
 			throw new NotImplementedException("GetRecordsInBucketAsync is not supported on a upstream reference store");
 		}
 
-		public IAsyncEnumerable<NamespaceId> GetNamespacesAsync()
+		public IAsyncEnumerable<NamespaceId> GetNamespacesAsync(CancellationToken cancellationToken)
 		{
 			throw new NotImplementedException("GetNamespaces is not supported on a upstream reference store");
 		}
 
-		public IAsyncEnumerable<BucketId> GetBucketsAsync(NamespaceId ns)
+		public IAsyncEnumerable<BucketId> GetBucketsAsync(NamespaceId ns, CancellationToken cancellationToken)
 		{
 			throw new NotImplementedException("GetBuckets is not supported on a upstream reference store");
 		}
 
-		public Task<bool> DeleteAsync(NamespaceId ns, BucketId bucket, RefId key)
+		public Task<bool> DeleteAsync(NamespaceId ns, BucketId bucket, RefId key, CancellationToken cancellationToken)
 		{
 			throw new NotImplementedException("Delete is not supported on a upstream reference store");
 		}
 
-		public Task<long> DropNamespaceAsync(NamespaceId ns)
+		public Task<long> DropNamespaceAsync(NamespaceId ns, CancellationToken cancellationToken)
 		{
 			throw new NotImplementedException("DropNamespace is not supported on a upstream reference store");
 		}
 
-		public Task<long> DeleteBucketAsync(NamespaceId ns, BucketId bucket)
+		public Task<long> DeleteBucketAsync(NamespaceId ns, BucketId bucket, CancellationToken cancellationToken)
 		{
 			throw new NotImplementedException("DeleteBucket is not supported on a upstream reference store");
 		}
diff --git a/Engine/Source/Programs/UnrealCloudDDC/Jupiter/Implementation/Objects/IReferencesStore.cs b/Engine/Source/Programs/UnrealCloudDDC/Jupiter/Implementation/Objects/IReferencesStore.cs
--- a/Engine/Source/Programs/UnrealCloudDDC/Jupiter/Implementation/Objects/IReferencesStore.cs
+++ b/Engine/Source/Programs/UnrealCloudDDC/Jupiter/Implementation/Objects/IReferencesStore.cs
@@ -2,6 +2,7 @@
 
 using System;
 using System.Collections.Generic;
+using System.Threading;
 using System.Threading.Tasks;
 using EpicGames.Horde.Storage;
 
@@ -9,7 +10,7 @@ namespace Jupiter.Implementation
 {
 	public interface IReferencesStore
 	{
-		Task<RefRecord> GetAsync(NamespaceId ns, BucketId bucket, RefId key, FieldFlags fieldFlags, OperationFlags opFlags);
+		Task<RefRecord> GetAsync(NamespaceId ns, BucketId bucket, RefId key, FieldFlags fieldFlags, OperationFlags opFlags, CancellationToken cancellationToken = default);
 
 		[Flags]
 		public enum FieldFlags
@@ -26,22 +27,22 @@ namespace Jupiter.Implementation
 			BypassCache = 1
 		}
 
-		Task PutAsync(NamespaceId ns, BucketId bucket, RefId key, BlobId blobHash, byte[] blob, bool isFinalized);
-		Task FinalizeAsync(NamespaceId ns, BucketId bucket, RefId key, BlobId blobIdentifier);
+		Task PutAsync(NamespaceId ns, BucketId bucket, RefId key, BlobId blobHash, byte[] blob, bool isFinalized, CancellationToken cancellationToken = default);
+		Task FinalizeAsync(NamespaceId ns, BucketId bucket, RefId key, BlobId blobIdentifier, CancellationToken cancellationToken = default);
 
-		Task<DateTime?> GetLastAccessTimeAsync(NamespaceId ns, BucketId bucket, RefId key);
-		Task UpdateLastAccessTimeAsync(NamespaceId ns, BucketId bucket, RefId key, DateTime newLastAccessTime);
-		IAsyncEnumerable<(NamespaceId, BucketId, RefId, DateTime)> GetRecordsAsync();
+		Task<DateTime?> GetLastAccessTimeAsync(NamespaceId ns, BucketId bucket, RefId key, CancellationToken cancellationToken = default);
+		Task UpdateLastAccessTimeAsync(NamespaceId ns, BucketId bucket, RefId key, DateTime newLastAccessTime, CancellationToken cancellationToken = default);
+		IAsyncEnumerable<(NamespaceId, BucketId, RefId, DateTime)> GetRecordsAsync(CancellationToken cancellationToken = default);
 
-		IAsyncEnumerable<(NamespaceId, BucketId, RefId)> GetRecordsWithoutAccessTimeAsync();
+		IAsyncEnumerable<(NamespaceId, BucketId, RefId)> GetRecordsWithoutAccessTimeAsync(CancellationToken cancellationToken = default);
 
-		IAsyncEnumerable<(RefId, BlobId)> GetRecordsInBucketAsync(NamespaceId ns, BucketId bucket);
+		IAsyncEnumerable<(RefId, BlobId)> GetRecordsInBucketAsync(NamespaceId ns, BucketId bucket, CancellationToken cancellationToken = default);
 
-		IAsyncEnumerable<NamespaceId> GetNamespacesAsync();
-		IAsyncEnumerable<BucketId> GetBucketsAsync(NamespaceId ns);
-		Task<bool> DeleteAsync(NamespaceId ns, BucketId bucket, RefId key);
-		Task<long> DropNamespaceAsync(NamespaceId ns);
-		Task<long> DeleteBucketAsync(NamespaceId ns, BucketId bucket);
+		IAsyncEnumerable<NamespaceId> GetNamespacesAsync(CancellationToken cancellationToken = default);
+		IAsyncEnumerable<BucketId> GetBucketsAsync(NamespaceId ns, CancellationToken cancellationToken = default);
+		Task<bool> DeleteAsync(NamespaceId ns, BucketId bucket, RefId key, CancellationToken cancellationToken = default);
+		Task<long> DropNamespaceAsync(NamespaceId ns, CancellationToken cancellationToken = default);
+		Task<long> DeleteBucketAsync(NamespaceId ns, BucketId bucket, CancellationToken cancellationToken = default);
 	}
 
 	public class RefRecord
diff --git a/Engine/Source/Programs/UnrealCloudDDC/Jupiter/Implementation/Objects/MemoryCachedReferencesStore.cs b/Engine/Source/Programs/UnrealCloudDDC/Jupiter/Implementation/Objects/MemoryCachedReferencesStore.cs
--- a/Engine/Source/Programs/UnrealCloudDDC/Jupiter/Implementation/Objects/MemoryCachedReferencesStore.cs
+++ b/Engine/Source/Programs/UnrealCloudDDC/Jupiter/Implementation/Objects/MemoryCachedReferencesStore.cs
@@ -3,6 +3,7 @@
 using System;
 using System.Collections.Concurrent;
 using System.Collections.Generic;
+using System.Threading;
 using System.Threading.Tasks;
 using EpicGames.Horde.Storage;
 using Microsoft.Extensions.Caching.Memory;
@@ -46,11 +47,11 @@ namespace Jupiter.Implementation.Objects
 			return _referenceCaches.GetOrAdd(ns, id => new MemoryCache(_options.CurrentValue));
 		}
 
-		public async Task<RefRecord> GetAsync(NamespaceId ns, BucketId bucket, RefId key, IReferencesStore.FieldFlags fieldFlags, IReferencesStore.OperationFlags opFlags)
+		public async Task<RefRecord> GetAsync(NamespaceId ns, BucketId bucket, RefId key, IReferencesStore.FieldFlags fieldFlags, IReferencesStore.OperationFlags opFlags, CancellationToken cancellationToken)
 		{
 			if (opFlags.HasFlag(IReferencesStore.OperationFlags.BypassCache))
 			{
-				return await _actualStore.GetAsync(ns, bucket, key, IReferencesStore.FieldFlags.All, opFlags);
+				return await _actualStore.GetAsync(ns, bucket, key, IReferencesStore.FieldFlags.All, opFlags, cancellationToken);
 			}
 
 			using TelemetrySpan scope = _tracer.StartActiveSpan("Ref.get")
@@ -72,29 +73,29 @@ namespace Jupiter.Implementation.Objects
 			}
 
 			scope.SetAttribute("Found", false);
-			RefRecord objectRecord = await _actualStore.GetAsync(ns, bucket, key, IReferencesStore.FieldFlags.All, opFlags);
+			RefRecord objectRecord = await _actualStore.GetAsync(ns, bucket, key, IReferencesStore.FieldFlags.All, opFlags, cancellationToken);
 			AddCacheEntry(ns, bucket, key, objectRecord);
 
 			return objectRecord;
 		}
 
-		public Task PutAsync(NamespaceId ns, BucketId bucket, RefId key, BlobId blobHash, byte[] blob, bool isFinalized)
+		public Task PutAsync(NamespaceId ns, BucketId bucket, RefId key, BlobId blobHash, byte[] blob, bool isFinalized, CancellationToken cancellationToken)
 		{
 			RefRecord objectRecord = new RefRecord(ns, bucket, key, DateTime.Now, blob, blobHash, isFinalized);
 			AddCacheEntry(ns, bucket, key, objectRecord);
 
-			return _actualStore.PutAsync(ns, bucket, key, blobHash, blob, isFinalized);
+			return _actualStore.PutAsync(ns, bucket, key, blobHash, blob, isFinalized, cancellationToken);
 		}
 
-		public Task FinalizeAsync(NamespaceId ns, BucketId bucket, RefId key, BlobId blobIdentifier)
+		public Task FinalizeAsync(NamespaceId ns, BucketId bucket, RefId key, BlobId blobIdentifier, CancellationToken cancellationToken)
 		{
 			FinalizeCacheEntry(ns, bucket, key);
-			return _actualStore.FinalizeAsync(ns, bucket, key, blobIdentifier);
+			return _actualStore.FinalizeAsync(ns, bucket, key, blobIdentifier, cancellationToken);
 		}
 
-		public Task<DateTime?> GetLastAccessTimeAsync(NamespaceId ns, BucketId bucket, RefId key)
+		public Task<DateTime?> GetLastAccessTimeAsync(NamespaceId ns, BucketId bucket, RefId key, CancellationToken cancellationToken)
 		{
-			return _actualStore.GetLastAccessTimeAsync(ns, bucket, key);
+			return _actualStore.GetLastAccessTimeAsync(ns, bucket, key, cancellationToken);
 		}
 
 		private void FinalizeCacheEntry(NamespaceId ns, BucketId bucket, RefId key)
@@ -107,55 +108,55 @@ namespace Jupiter.Implementation.Objects
 			}
 		}
 
-		public Task UpdateLastAccessTimeAsync(NamespaceId ns, BucketId bucket, RefId key, DateTime newLastAccessTime)
+		public Task UpdateLastAccessTimeAsync(NamespaceId ns, BucketId bucket, RefId key, DateTime newLastAccessTime, CancellationToken cancellationToken)
 		{
-			return _actualStore.UpdateLastAccessTimeAsync(ns, bucket, key, newLastAccessTime);
+			return _actualStore.UpdateLastAccessTimeAsync(ns, bucket, key, newLastAccessTime, cancellationToken);
 		}
 
-		public IAsyncEnumerable<(NamespaceId, BucketId, RefId, DateTime)> GetRecordsAsync()
+		public IAsyncEnumerable<(NamespaceId, BucketId, RefId, DateTime)> GetRecordsAsync(CancellationToken cancellationToken)
 		{
-			return _actualStore.GetRecordsAsync();
+			return _actualStore.GetRecordsAsync(cancellationToken);
 		}
 
-		public IAsyncEnumerable<(NamespaceId, BucketId, RefId)> GetRecordsWithoutAccessTimeAsync()
+		public IAsyncEnumerable<(NamespaceId, BucketId, RefId)> GetRecordsWithoutAccessTimeAsync(CancellationToken cancellationToken)
 		{
-			return _actualStore.GetRecordsWithoutAccessTimeAsync();
+			return _actualStore.GetRecordsWithoutAccessTimeAsync(cancellationToken);
 		}
 
-		public IAsyncEnumerable<(RefId, BlobId)> GetRecordsInBucketAsync(NamespaceId ns, BucketId bucket)
+		public IAsyncEnumerable<(RefId, BlobId)> GetRecordsInBucketAsync(NamespaceId ns, BucketId bucket, CancellationToken cancellationToken)
 		{
-			return _actualStore.GetRecordsInBucketAsync(ns, bucket);
+			return _actualStore.GetRecordsInBucketAsync(ns, bucket, cancellationToken);
 		}
 
-		public IAsyncEnumerable<NamespaceId> GetNamespacesAsync()
+		public IAsyncEnumerable<NamespaceId> GetNamespacesAsync(CancellationToken cancellationToken)
 		{
-			return _actualStore.GetNamespacesAsync();
+			return _actualStore.GetNamespacesAsync(cancellationToken);
 		}
 
-		public IAsyncEnumerable<BucketId> GetBucketsAsync(NamespaceId ns)
+		public IAsyncEnumerable<BucketId> GetBucketsAsync(NamespaceId ns, CancellationToken cancellationToken)
 		{
-			return _actualStore.GetBucketsAsync(ns);
+			return _actualStore.GetBucketsAsync(ns, cancellationToken);
 		}
 
-		public Task<bool> DeleteAsync(NamespaceId ns, BucketId bucket, RefId key)
+		public Task<bool> DeleteAsync(NamespaceId ns, BucketId bucket, RefId key, CancellationToken cancellationToken)
 		{
 			MemoryCache cache = GetCacheForNamespace(ns);
 			cache.Remove(new CachedReferenceKey(bucket, key));
-			return _actualStore.DeleteAsync(ns, bucket, key);
+			return _actualStore.DeleteAsync(ns, bucket, key, cancellationToken);
 		}
 
-		public Task<long> DropNamespaceAsync(NamespaceId ns)
+		public Task<long> DropNamespaceAsync(NamespaceId ns, CancellationToken cancellationToken)
 		{
 			_referenceCaches.TryRemove(ns, out _);
-			return _actualStore.DropNamespaceAsync(ns);
+			return _actualStore.DropNamespaceAsync(ns, cancellationToken);
 		}
 
-		public Task<long> DeleteBucketAsync(NamespaceId ns, BucketId bucket)
+		public Task<long> DeleteBucketAsync(NamespaceId ns, BucketId bucket, CancellationToken cancellationToken)
 		{
 			// we do not track enough information to be able to drop a bucket, so we have to drop the entire namespace cache to remove the bucket
 			// this should be okay as deleting buckets is a extremely uncommon operation
 			_referenceCaches.TryRemove(ns, out _);
-			return _actualStore.DeleteBucketAsync(ns, bucket);
+			return _actualStore.DeleteBucketAsync(ns, bucket, cancellationToken);
 		}
 
 		public void Clear()
diff --git a/Engine/Source/Programs/UnrealCloudDDC/Jupiter/Implementation/Objects/MemoryReferencesStore.cs b/Engine/Source/Programs/UnrealCloudDDC/Jupiter/Implementation/Objects/MemoryReferencesStore.cs
--- a/Engine/Source/Programs/UnrealCloudDDC/Jupiter/Implementation/Objects/MemoryReferencesStore.cs
+++ b/Engine/Source/Programs/UnrealCloudDDC/Jupiter/Implementation/Objects/MemoryReferencesStore.cs
@@ -4,6 +4,8 @@ using System;
 using System.Collections.Concurrent;
 using System.Collections.Generic;
 using System.Linq;
+using System.Runtime.CompilerServices;
+using System.Threading;
 using System.Threading.Tasks;
 using EpicGames.Horde.Storage;
 
@@ -19,7 +21,7 @@ namespace Jupiter.Implementation
 
 		}
 
-		public Task<RefRecord> GetAsync(NamespaceId ns, BucketId bucket, RefId key, IReferencesStore.FieldFlags fieldFlags, IReferencesStore.OperationFlags opFlags)
+		public Task<RefRecord> GetAsync(NamespaceId ns, BucketId bucket, RefId key, IReferencesStore.FieldFlags fieldFlags, IReferencesStore.OperationFlags opFlags, CancellationToken cancellationToken)
 		{
 			if (_objects.TryGetValue(BuildKey(ns, bucket, key), out MemoryStoreObject? o))
 			{
@@ -29,7 +31,7 @@ namespace Jupiter.Implementation
 			throw new RefNotFoundException(ns, bucket, key);
 		}
 
-		public Task PutAsync(NamespaceId ns, BucketId bucket, RefId key, BlobId blobHash, byte[] blob, bool isFinalized)
+		public Task PutAsync(NamespaceId ns, BucketId bucket, RefId key, BlobId blobHash, byte[] blob, bool isFinalized, CancellationToken cancellationToken)
 		{
 			lock (_namespaces)
 			{
@@ -43,7 +45,7 @@ namespace Jupiter.Implementation
 			return Task.FromResult(o);
 		}
 
-		public Task FinalizeAsync(NamespaceId ns, BucketId bucket, RefId key, BlobId blobIdentifier)
+		public Task FinalizeAsync(NamespaceId ns, BucketId bucket, RefId key, BlobId blobIdentifier, CancellationToken cancellationToken)
 		{
 			if (!_objects.TryGetValue(BuildKey(ns, bucket, key), out MemoryStoreObject? o))
 			{
@@ -54,7 +56,7 @@ namespace Jupiter.Implementation
 			return Task.CompletedTask;
 		}
 
-		public Task<DateTime?> GetLastAccessTimeAsync(NamespaceId ns, BucketId bucket, RefId key)
+		public Task<DateTime?> GetLastAccessTimeAsync(NamespaceId ns, BucketId bucket, RefId key, CancellationToken cancellationToken)
 		{
 			DateTime? lastAccessTime = null;
 			if (_objects.TryGetValue(BuildKey(ns, bucket, key), out MemoryStoreObject? o))
@@ -65,7 +67,7 @@ namespace Jupiter.Implementation
 			return Task.FromResult(lastAccessTime);
 		}
 
-		public Task UpdateLastAccessTimeAsync(NamespaceId ns, BucketId bucket, RefId key, DateTime lastAccessTime)
+		public Task UpdateLastAccessTimeAsync(NamespaceId ns, BucketId bucket, RefId key, DateTime lastAccessTime, CancellationToken cancellationToken)
 		{
 			if (!_objects.TryGetValue(BuildKey(ns, bucket, key), out MemoryStoreObject? o))
 			{
@@ -76,16 +78,16 @@ namespace Jupiter.Implementation
 			return Task.CompletedTask;
 		}
 
-		public async IAsyncEnumerable<(NamespaceId, BucketId, RefId)> GetRecordsWithoutAccessTimeAsync()
+		public async IAsyncEnumerable<(NamespaceId, BucketId, RefId)> GetRecordsWithoutAccessTimeAsync([EnumeratorCancellation] CancellationToken cancellationToken)
 		{
-			await foreach ((NamespaceId namespaceId, BucketId bucketId, RefId refId, DateTime _) in GetRecordsAsync())
+			await foreach ((NamespaceId namespaceId, BucketId bucketId, RefId refId, DateTime _) in GetRecordsAsync(cancellationToken))
 			{
 				await Task.CompletedTask;
 				yield return (namespaceId, bucketId, refId);
 			}
 		}
 
-		public async IAsyncEnumerable<(NamespaceId, BucketId, RefId, DateTime)> GetRecordsAsync()
+		public async IAsyncEnumerable<(NamespaceId, BucketId, RefId, DateTime)> GetRecordsAsync([EnumeratorCancellation] CancellationToken cancellationToken)
 		{
 			foreach (MemoryStoreObject o in _objects.Values.OrderBy(o => o.LastAccessTime))
 			{
@@ -94,7 +96,7 @@ namespace Jupiter.Implementation
 			}
 		}
 
-		public async IAsyncEnumerable<(RefId, BlobId)> GetRecordsInBucketAsync(NamespaceId ns, BucketId bucket)
+		public async IAsyncEnumerable<(RefId, BlobId)> GetRecordsInBucketAsync(NamespaceId ns, BucketId bucket, [EnumeratorCancellation] CancellationToken cancellationToken)
 		{
 			foreach (MemoryStoreObject o in _objects.Values.Where(o => o.Namespace == ns && o.Bucket == bucket))
 			{
@@ -103,12 +105,12 @@ namespace Jupiter.Implementation
 			}
 		}
 
-		public IAsyncEnumerable<NamespaceId> GetNamespacesAsync()
+		public IAsyncEnumerable<NamespaceId> GetNamespacesAsync(CancellationToken cancellationToken)
 		{
 			return _namespaces.ToAsyncEnumerable();
 		}
 
-		public async IAsyncEnumerable<BucketId> GetBucketsAsync(NamespaceId ns)
+		public async IAsyncEnumerable<BucketId> GetBucketsAsync(NamespaceId ns, [EnumeratorCancellation] CancellationToken cancellationToken)
 		{
 			HashSet<BucketId> buckets = new HashSet<BucketId>();
 			foreach (MemoryStoreObject o in _objects.Values.Where(o => o.Namespace == ns))
@@ -123,7 +125,7 @@ namespace Jupiter.Implementation
 			}
 		}
 
-		public Task<bool> DeleteAsync(NamespaceId ns, BucketId bucket, RefId key)
+		public Task<bool> DeleteAsync(NamespaceId ns, BucketId bucket, RefId key, CancellationToken cancellationToken)
 		{
 			if (!_objects.TryRemove(BuildKey(ns, bucket, key), out MemoryStoreObject? _))
 			{
@@ -133,7 +135,7 @@ namespace Jupiter.Implementation
 			return Task.FromResult(true);
 		}
 
-		public Task<long> DropNamespaceAsync(NamespaceId ns)
+		public Task<long> DropNamespaceAsync(NamespaceId ns, CancellationToken cancellationToken)
 		{
 			lock (_namespaces)
 			{
@@ -162,7 +164,7 @@ namespace Jupiter.Implementation
 			return Task.FromResult(removedCount);
 		}
 
-		public Task<long> DeleteBucketAsync(NamespaceId ns, BucketId bucket)
+		public Task<long> DeleteBucketAsync(NamespaceId ns, BucketId bucket, CancellationToken cancellationToken)
 		{
 			List<string> objectToRemove = new List<string>();
 
diff --git a/Engine/Source/Programs/UnrealCloudDDC/Jupiter/Implementation/Objects/MongoReferencesStore.cs b/Engine/Source/Programs/UnrealCloudDDC/Jupiter/Implementation/Objects/MongoReferencesStore.cs
--- a/Engine/Source/Programs/UnrealCloudDDC/Jupiter/Implementation/Objects/MongoReferencesStore.cs
+++ b/Engine/Source/Programs/UnrealCloudDDC/Jupiter/Implementation/Objects/MongoReferencesStore.cs
@@ -2,6 +2,8 @@
 
 using System;
 using System.Collections.Generic;
+using System.Runtime.CompilerServices;
+using System.Threading;
 using System.Threading.Tasks;
 using EpicGames.Horde.Storage;
 using Jupiter.Common;
@@ -47,12 +49,12 @@ namespace Jupiter.Implementation
 			AddIndexFor<MongoReferencesModelV0>().CreateOne(indexTTL);
 		}
 
-		public async Task<RefRecord> GetAsync(NamespaceId ns, BucketId bucket, RefId key, IReferencesStore.FieldFlags fieldFlags, IReferencesStore.OperationFlags opFlags)
+		public async Task<RefRecord> GetAsync(NamespaceId ns, BucketId bucket, RefId key, IReferencesStore.FieldFlags fieldFlags, IReferencesStore.OperationFlags opFlags, CancellationToken cancellationToken)
 		{
 			bool includePayload = (fieldFlags & IReferencesStore.FieldFlags.IncludePayload) != 0;
 			IMongoCollection<MongoReferencesModelV0> collection = GetCollection<MongoReferencesModelV0>();
-			IAsyncCursor<MongoReferencesModelV0>? cursor = await collection.FindAsync(m => m.Ns == ns.ToString() && m.Bucket == bucket.ToString() && m.Key == key.ToString());
-			MongoReferencesModelV0? model = await cursor.FirstOrDefaultAsync();
+			IAsyncCursor<MongoReferencesModelV0>? cursor = await collection.FindAsync(m => m.Ns == ns.ToString() && m.Bucket == bucket.ToString() && m.Key == key.ToString(), cancellationToken: cancellationToken);
+			MongoReferencesModelV0? model = await cursor.FirstOrDefaultAsync(cancellationToken);
 			if (model == null)
 			{
 				throw new RefNotFoundException(ns, bucket, key);
@@ -67,7 +69,7 @@ namespace Jupiter.Implementation
 			return model.ToRefRecord();
 		}
 
-		public async Task PutAsync(NamespaceId ns, BucketId bucket, RefId key, BlobId blobHash, byte[]? blob, bool isFinalized)
+		public async Task PutAsync(NamespaceId ns, BucketId bucket, RefId key, BlobId blobHash, byte[]? blob, bool isFinalized, CancellationToken cancellationToken)
 		{
 			IMongoCollection<MongoReferencesModelV0> collection = GetCollection<MongoReferencesModelV0>();
 
@@ -78,7 +80,7 @@ namespace Jupiter.Implementation
 				blob = Array.Empty<byte>();
 			}
 
-			Task addNamespaceTask = AddNamespaceIfNotExistAsync(ns);
+			Task addNamespaceTask = AddNamespaceIfNotExistAsync(ns, cancellationToken);
 			MongoReferencesModelV0 model = new MongoReferencesModelV0(ns, bucket, key, blobHash, blob, isFinalized, DateTime.Now);
 
 			NamespacePolicy policy = _namespacePolicyResolver.GetPoliciesForNs(ns);
@@ -93,46 +95,48 @@ namespace Jupiter.Implementation
 			{
 				IsUpsert = true
 			};
-			await collection.FindOneAndReplaceAsync(filter, model, options);
+			await collection.FindOneAndReplaceAsync(filter, model, options, cancellationToken);
 
 			await addNamespaceTask;
 		}
 
-		public async Task FinalizeAsync(NamespaceId ns, BucketId bucket, RefId key, BlobId blobIdentifier)
+		public async Task FinalizeAsync(NamespaceId ns, BucketId bucket, RefId key, BlobId blobIdentifier, CancellationToken cancellationToken)
 		{
 			IMongoCollection<MongoReferencesModelV0> collection = GetCollection<MongoReferencesModelV0>();
 
 			UpdateResult _ = await collection.UpdateOneAsync(
 				model => model.Ns == ns.ToString() && model.Bucket == bucket.ToString() && model.Key == key.ToString(),
-				Builders<MongoReferencesModelV0>.Update.Set(model => model.IsFinalized, true)
+				Builders<MongoReferencesModelV0>.Update.Set(model => model.IsFinalized, true),
+				cancellationToken: cancellationToken
 			);
 		}
 
-		public async Task<DateTime?> GetLastAccessTimeAsync(NamespaceId ns, BucketId bucket, RefId key)
+		public async Task<DateTime?> GetLastAccessTimeAsync(NamespaceId ns, BucketId bucket, RefId key, CancellationToken cancellationToken)
 		{
 			IMongoCollection<MongoReferencesModelV0> collection = GetCollection<MongoReferencesModelV0>();
-			IAsyncCursor<MongoReferencesModelV0>? cursor = await collection.FindAsync(m => m.Ns == ns.ToString() && m.Bucket == bucket.ToString() && m.Key == key.ToString());
-			MongoReferencesModelV0? model = await cursor.FirstOrDefaultAsync();
+			IAsyncCursor<MongoReferencesModelV0>? cursor = await collection.FindAsync(m => m.Ns == ns.ToString() && m.Bucket == bucket.ToString() && m.Key == key.ToString(), cancellationToken: cancellationToken);
+			MongoReferencesModelV0? model = await cursor.FirstOrDefaultAsync(cancellationToken);
 
 			return model?.LastAccessTime;
 		}
 
-		public async Task UpdateLastAccessTimeAsync(NamespaceId ns, BucketId bucket, RefId key, DateTime newLastAccessTime)
+		public async Task UpdateLastAccessTimeAsync(NamespaceId ns, BucketId bucket, RefId key, DateTime newLastAccessTime, CancellationToken cancellationToken)
 		{
 			IMongoCollection<MongoReferencesModelV0> collection = GetCollection<MongoReferencesModelV0>();
 
 			UpdateResult _ = await collection.UpdateOneAsync(
 				model => model.Ns == ns.ToString() && model.Bucket == bucket.ToString() && model.Key == key.ToString(),
-				Builders<MongoReferencesModelV0>.Update.Set(model => model.LastAccessTime, newLastAccessTime)
+				Builders<MongoReferencesModelV0>.Update.Set(model => model.LastAccessTime, newLastAccessTime), 
+				cancellationToken: cancellationToken
 			);
 		}
 
-		public async IAsyncEnumerable<(NamespaceId, BucketId, RefId, DateTime)> GetRecordsAsync()
+		public async IAsyncEnumerable<(NamespaceId, BucketId, RefId, DateTime)> GetRecordsAsync([EnumeratorCancellation] CancellationToken cancellationToken)
 		{
 			IMongoCollection<MongoReferencesModelV0> collection = GetCollection<MongoReferencesModelV0>();
-			IAsyncCursor<MongoReferencesModelV0>? cursor = await collection.FindAsync(FilterDefinition<MongoReferencesModelV0>.Empty);
+			IAsyncCursor<MongoReferencesModelV0>? cursor = await collection.FindAsync(FilterDefinition<MongoReferencesModelV0>.Empty, cancellationToken: cancellationToken);
 
-			while (await cursor.MoveNextAsync())
+			while (await cursor.MoveNextAsync(cancellationToken))
 			{
 				foreach (MongoReferencesModelV0 model in cursor.Current)
 				{
@@ -141,20 +145,20 @@ namespace Jupiter.Implementation
 			}
 		}
 
-		public async IAsyncEnumerable<(NamespaceId, BucketId, RefId)> GetRecordsWithoutAccessTimeAsync()
+		public async IAsyncEnumerable<(NamespaceId, BucketId, RefId)> GetRecordsWithoutAccessTimeAsync([EnumeratorCancellation] CancellationToken cancellationToken)
 		{
-			await foreach ((NamespaceId ns, BucketId bucket, RefId key, DateTime _) in GetRecordsAsync())
+			await foreach ((NamespaceId ns, BucketId bucket, RefId key, DateTime _) in GetRecordsAsync(cancellationToken))
 			{
 				yield return (ns, bucket, key);
 			}
 		}
 
-		public async IAsyncEnumerable<(RefId, BlobId)> GetRecordsInBucketAsync(NamespaceId ns, BucketId bucket)
+		public async IAsyncEnumerable<(RefId, BlobId)> GetRecordsInBucketAsync(NamespaceId ns, BucketId bucket, [EnumeratorCancellation] CancellationToken cancellationToken)
 		{
 			IMongoCollection<MongoReferencesModelV0> collection = GetCollection<MongoReferencesModelV0>();
-			IAsyncCursor<MongoReferencesModelV0>? cursor = await collection.FindAsync(m => m.Ns == ns.ToString() && m.Bucket == bucket.ToString());
+			IAsyncCursor<MongoReferencesModelV0>? cursor = await collection.FindAsync(m => m.Ns == ns.ToString() && m.Bucket == bucket.ToString(), cancellationToken: cancellationToken);
 
-			while (await cursor.MoveNextAsync())
+			while (await cursor.MoveNextAsync(cancellationToken))
 			{
 				foreach (MongoReferencesModelV0 model in cursor.Current)
 				{
@@ -163,7 +167,7 @@ namespace Jupiter.Implementation
 			}
 		}
 
-		public async Task AddNamespaceIfNotExistAsync(NamespaceId ns)
+		public async Task AddNamespaceIfNotExistAsync(NamespaceId ns, CancellationToken cancellationToken)
 		{
 			FilterDefinition<MongoNamespacesModelV0> filter = Builders<MongoNamespacesModelV0>.Filter.Where(m => m.Ns == ns.ToString());
 			FindOneAndReplaceOptions<MongoNamespacesModelV0, MongoNamespacesModelV0> options = new FindOneAndReplaceOptions<MongoNamespacesModelV0, MongoNamespacesModelV0>
@@ -172,15 +176,15 @@ namespace Jupiter.Implementation
 			};
 
 			IMongoCollection<MongoNamespacesModelV0> collection = GetCollection<MongoNamespacesModelV0>();
-			await collection.FindOneAndReplaceAsync(filter, new MongoNamespacesModelV0(ns), options);
+			await collection.FindOneAndReplaceAsync(filter, new MongoNamespacesModelV0(ns), options, cancellationToken);
 		}
 
-		public async IAsyncEnumerable<NamespaceId> GetNamespacesAsync()
+		public async IAsyncEnumerable<NamespaceId> GetNamespacesAsync([EnumeratorCancellation] CancellationToken cancellationToken)
 		{
 			IMongoCollection<MongoNamespacesModelV0> collection = GetCollection<MongoNamespacesModelV0>();
 
-			IAsyncCursor<MongoNamespacesModelV0> cursor = await collection.FindAsync(FilterDefinition<MongoNamespacesModelV0>.Empty);
-			while (await cursor.MoveNextAsync())
+			IAsyncCursor<MongoNamespacesModelV0> cursor = await collection.FindAsync(FilterDefinition<MongoNamespacesModelV0>.Empty, cancellationToken: cancellationToken);
+			while (await cursor.MoveNextAsync(cancellationToken))
 			{
 				foreach (MongoNamespacesModelV0? document in cursor.Current)
 				{
@@ -189,14 +193,14 @@ namespace Jupiter.Implementation
 			}
 		}
 
-		public async IAsyncEnumerable<BucketId> GetBucketsAsync(NamespaceId ns)
+		public async IAsyncEnumerable<BucketId> GetBucketsAsync(NamespaceId ns, [EnumeratorCancellation] CancellationToken cancellationToken)
 		{
 			IMongoCollection<MongoReferencesModelV0> collection = GetCollection<MongoReferencesModelV0>();
 
-			IAsyncCursor<MongoReferencesModelV0> cursor = await collection.FindAsync(m => m.Ns == ns.ToString());
+			IAsyncCursor<MongoReferencesModelV0> cursor = await collection.FindAsync(m => m.Ns == ns.ToString(), cancellationToken: cancellationToken);
 
 			HashSet<BucketId> buckets = new HashSet<BucketId>();
-			while (await cursor.MoveNextAsync())
+			while (await cursor.MoveNextAsync(cancellationToken))
 			{
 				foreach (MongoReferencesModelV0? document in cursor.Current)
 				{
@@ -210,22 +214,23 @@ namespace Jupiter.Implementation
 			}
 		}
 
-		public async Task<bool> DeleteAsync(NamespaceId ns, BucketId bucket, RefId key)
+		public async Task<bool> DeleteAsync(NamespaceId ns, BucketId bucket, RefId key, CancellationToken cancellationToken)
 		{
 			IMongoCollection<MongoReferencesModelV0> collection = GetCollection<MongoReferencesModelV0>();
 
 			DeleteResult result = await collection.DeleteOneAsync(model =>
-				model.Ns == ns.ToString() && model.Bucket == bucket.ToString() && model.Key == key.ToString());
+				model.Ns == ns.ToString() && model.Bucket == bucket.ToString() && model.Key == key.ToString(),
+				cancellationToken);
 
 			return result.DeletedCount != 0;
 		}
 
-		public async Task<long> DropNamespaceAsync(NamespaceId ns)
+		public async Task<long> DropNamespaceAsync(NamespaceId ns, CancellationToken cancellationToken)
 		{
 			IMongoCollection<MongoReferencesModelV0> collection = GetCollection<MongoReferencesModelV0>();
 
 			DeleteResult result = await collection.DeleteManyAsync(model =>
-				model.Ns == ns.ToString());
+				model.Ns == ns.ToString(), cancellationToken);
 
 			long deletedCount = 0;
 			if (result.IsAcknowledged)
@@ -234,17 +239,17 @@ namespace Jupiter.Implementation
 			}
 
 			IMongoCollection<MongoNamespacesModelV0> namespaceCollection = GetCollection<MongoNamespacesModelV0>();
-			await namespaceCollection.DeleteOneAsync(m => m.Ns == ns.ToString());
+			await namespaceCollection.DeleteOneAsync(m => m.Ns == ns.ToString(), cancellationToken);
 			return deletedCount;
 
 		}
 
-		public async Task<long> DeleteBucketAsync(NamespaceId ns, BucketId bucket)
+		public async Task<long> DeleteBucketAsync(NamespaceId ns, BucketId bucket, CancellationToken cancellationToken)
 		{
 			IMongoCollection<MongoReferencesModelV0> collection = GetCollection<MongoReferencesModelV0>();
 
 			DeleteResult result = await collection.DeleteManyAsync(model =>
-				model.Ns == ns.ToString() && model.Bucket == bucket.ToString());
+				model.Ns == ns.ToString() && model.Bucket == bucket.ToString(), cancellationToken);
 
 			if (result.IsAcknowledged)
 			{
diff --git a/Engine/Source/Programs/UnrealCloudDDC/Jupiter/Implementation/Objects/ScyllaReferencesStore.cs b/Engine/Source/Programs/UnrealCloudDDC/Jupiter/Implementation/Objects/ScyllaReferencesStore.cs
--- a/Engine/Source/Programs/UnrealCloudDDC/Jupiter/Implementation/Objects/ScyllaReferencesStore.cs
+++ b/Engine/Source/Programs/UnrealCloudDDC/Jupiter/Implementation/Objects/ScyllaReferencesStore.cs
@@ -4,6 +4,8 @@ using System;
 using System.Collections.Concurrent;
 using System.Collections.Generic;
 using System.Linq;
+using System.Runtime.CompilerServices;
+using System.Threading;
 using System.Threading.Tasks;
 using Cassandra;
 using Cassandra.Mapping;
@@ -97,7 +99,7 @@ namespace Jupiter.Implementation
 			_getObjectsInBucketPartitionRangeStatement = _session.Prepare($"SELECT name, payload_hash FROM objects WHERE namespace = ? AND bucket = ? ALLOW FILTERING {cqlOptions}");
 		}
 
-		public async Task<RefRecord> GetAsync(NamespaceId ns, BucketId bucket, RefId name, IReferencesStore.FieldFlags fieldFlags, IReferencesStore.OperationFlags opFlags)
+		public async Task<RefRecord> GetAsync(NamespaceId ns, BucketId bucket, RefId name, IReferencesStore.FieldFlags fieldFlags, IReferencesStore.OperationFlags opFlags, CancellationToken cancellationToken)
 		{
 			using TelemetrySpan scope = _tracer.BuildScyllaSpan("scylla.get").SetAttribute("resource.name", $"{ns}.{bucket}.{name}");
 			scope.SetAttribute("BypassCache", false);
@@ -140,7 +142,7 @@ namespace Jupiter.Implementation
 			return new RefRecord(new NamespaceId(o.Namespace!), new BucketId(o.Bucket!), new RefId(o.Name!), o.LastAccessTime, o.InlinePayload, o.PayloadHash!.AsBlobIdentifier(), o.IsFinalized!.Value);
 		}
 
-		public async Task PutAsync(NamespaceId ns, BucketId bucket, RefId name, BlobId blobHash, byte[] blob, bool isFinalized)
+		public async Task PutAsync(NamespaceId ns, BucketId bucket, RefId name, BlobId blobHash, byte[] blob, bool isFinalized, CancellationToken cancellationToken)
 		{
 			using TelemetrySpan scope = _tracer.BuildScyllaSpan("scylla.put").SetAttribute("resource.name", $"{ns}.{bucket}.{name}");
 
@@ -173,21 +175,21 @@ namespace Jupiter.Implementation
 			await addBucketTask;
 		}
 
-		public async Task FinalizeAsync(NamespaceId ns, BucketId bucket, RefId name, BlobId blobIdentifier)
+		public async Task FinalizeAsync(NamespaceId ns, BucketId bucket, RefId name, BlobId blobIdentifier, CancellationToken cancellationToken)
 		{
 			using TelemetrySpan scope = _tracer.BuildScyllaSpan("scylla.finalize").SetAttribute("resource.name", $"{ns}.{bucket}.{name}");
 
 			await _mapper.UpdateAsync<ScyllaObject>("SET is_finalized=true WHERE namespace=? AND bucket=? AND name=?", ns.ToString(), bucket.ToString(), name.ToString());
 		}
 
-		public async Task<DateTime?> GetLastAccessTimeAsync(NamespaceId ns, BucketId bucket, RefId key)
+		public async Task<DateTime?> GetLastAccessTimeAsync(NamespaceId ns, BucketId bucket, RefId key, CancellationToken cancellationToken)
 		{
 			using TelemetrySpan scope = _tracer.BuildScyllaSpan("scylla.get_last_access_time").SetAttribute("resource.name", $"{ns}.{bucket}.{key}");
 			ScyllaObjectLastAccess? lastAccessRecord = await _mapper.SingleOrDefaultAsync<ScyllaObjectLastAccess>("WHERE namespace = ? AND bucket = ? AND name = ?", ns.ToString(), bucket.ToString(), key.ToString());
 			return lastAccessRecord?.LastAccessTime;
 		}
 
-		public async Task UpdateLastAccessTimeAsync(NamespaceId ns, BucketId bucket, RefId name, DateTime lastAccessTime)
+		public async Task UpdateLastAccessTimeAsync(NamespaceId ns, BucketId bucket, RefId name, DateTime lastAccessTime, CancellationToken cancellationToken)
 		{
 			using TelemetrySpan scope = _tracer.BuildScyllaSpan("scylla.update_last_access_time");
 
@@ -202,7 +204,7 @@ namespace Jupiter.Implementation
 			await updateObjectLastAccessTask;
 		}
 
-		public async IAsyncEnumerable<(NamespaceId, BucketId, RefId, DateTime)> GetRecordsAsync()
+		public async IAsyncEnumerable<(NamespaceId, BucketId, RefId, DateTime)> GetRecordsAsync([EnumeratorCancellation] CancellationToken cancellationToken)
 		{
 			using TelemetrySpan scope = _tracer.BuildScyllaSpan("scylla.get_records");
 
@@ -264,7 +266,7 @@ namespace Jupiter.Implementation
 								"Cassandra read timeouts, waiting a while and then retrying. Attempt {Attempts} .",
 								retryAttempts);
 							// wait 10 seconds and try again as the Db is under heavy load right now
-							await Task.Delay(TimeSpan.FromSeconds(10));
+							await Task.Delay(TimeSpan.FromSeconds(10), cancellationToken);
 							timeoutException = e;
 						}
 					}
@@ -280,7 +282,7 @@ namespace Jupiter.Implementation
 			}
 		}
 
-		public async IAsyncEnumerable<(NamespaceId, BucketId, RefId)> GetRecordsWithoutAccessTimeAsync()
+		public async IAsyncEnumerable<(NamespaceId, BucketId, RefId)> GetRecordsWithoutAccessTimeAsync([EnumeratorCancellation] CancellationToken cancellationToken)
 		{
 			using TelemetrySpan scope = _tracer.BuildScyllaSpan("scylla.get_records_no_access_time");
 
@@ -338,7 +340,7 @@ namespace Jupiter.Implementation
 								"Cassandra read timeouts, waiting a while and then retrying. Attempt {Attempts} .",
 								retryAttempts);
 							// wait 10 seconds and try again as the Db is under heavy load right now
-							await Task.Delay(TimeSpan.FromSeconds(10));
+							await Task.Delay(TimeSpan.FromSeconds(10), cancellationToken);
 							timeoutException = e;
 						}
 					}
@@ -354,7 +356,7 @@ namespace Jupiter.Implementation
 			}
 		}
 
-		public async IAsyncEnumerable<(RefId, BlobId)> GetRecordsInBucketAsync(NamespaceId ns, BucketId bucket)
+		public async IAsyncEnumerable<(RefId, BlobId)> GetRecordsInBucketAsync(NamespaceId ns, BucketId bucket, [EnumeratorCancellation] CancellationToken cancellationToken)
 		{
 			using TelemetrySpan scope = _tracer.BuildScyllaSpan("scylla.get_records_in_bucket_per_shard");
 			PreparedStatement getObjectStatement = _getObjectsInBucketPartitionRangeStatement;
@@ -476,7 +478,7 @@ namespace Jupiter.Implementation
 			}
 		}
 
-		public async IAsyncEnumerable<NamespaceId> GetNamespacesAsync()
+		public async IAsyncEnumerable<NamespaceId> GetNamespacesAsync([EnumeratorCancellation] CancellationToken cancellationToken)
 		{
 			using TelemetrySpan scope = _tracer.BuildScyllaSpan("scylla.get_namespaces");
 
@@ -513,7 +515,7 @@ namespace Jupiter.Implementation
 			}
 		}
 
-		public async IAsyncEnumerable<BucketId> GetBucketsAsync(NamespaceId ns)
+		public async IAsyncEnumerable<BucketId> GetBucketsAsync(NamespaceId ns, [EnumeratorCancellation] CancellationToken cancellationToken)
 		{
 			foreach (ScyllaBucket scyllaBucket in await _mapper.FetchAsync<ScyllaBucket>("WHERE namespace=?", ns.ToString()))
 			{
@@ -526,7 +528,7 @@ namespace Jupiter.Implementation
 			}
 		}
 
-		public async Task<bool> DeleteAsync(NamespaceId ns, BucketId bucket, RefId key)
+		public async Task<bool> DeleteAsync(NamespaceId ns, BucketId bucket, RefId key, CancellationToken cancellationToken)
 		{
 			using TelemetrySpan scope = _tracer.BuildScyllaSpan("scylla.delete_record").SetAttribute("resource.name", $"{ns}.{bucket}.{key}");
 
@@ -539,7 +541,7 @@ namespace Jupiter.Implementation
 			return true;
 		}
 
-		public async Task<long> DropNamespaceAsync(NamespaceId ns)
+		public async Task<long> DropNamespaceAsync(NamespaceId ns, CancellationToken cancellationToken)
 		{
 			using TelemetrySpan scope = _tracer.BuildScyllaSpan("scylla.delete_namespace");
 			RowSet rowSet = await _session.ExecuteAsync(new SimpleStatement("SELECT bucket, name FROM objects WHERE namespace = ? ALLOW FILTERING;", ns.ToString()));
@@ -549,7 +551,7 @@ namespace Jupiter.Implementation
 				string bucket = row.GetValue<string>("bucket");
 				string name = row.GetValue<string>("name");
 
-				await DeleteAsync(ns, new BucketId(bucket), new RefId(name));
+				await DeleteAsync(ns, new BucketId(bucket), new RefId(name), cancellationToken);
 
 				deletedCount++;
 			}
@@ -561,7 +563,7 @@ namespace Jupiter.Implementation
 			return deletedCount;
 		}
 
-		public async Task<long> DeleteBucketAsync(NamespaceId ns, BucketId bucket)
+		public async Task<long> DeleteBucketAsync(NamespaceId ns, BucketId bucket, CancellationToken cancellationToken)
 		{
 			using TelemetrySpan scope = _tracer.BuildScyllaSpan("scylla.delete_bucket");
 
@@ -571,7 +573,7 @@ namespace Jupiter.Implementation
 			{
 				string name = row.GetValue<string>("name");
 
-				await DeleteAsync(ns, bucket, new RefId(name));
+				await DeleteAsync(ns, bucket, new RefId(name), cancellationToken);
 				deletedCount++;
 			}
 
diff --git a/Engine/Source/Programs/UnrealCloudDDC/Jupiter/Implementation/RefService.cs b/Engine/Source/Programs/UnrealCloudDDC/Jupiter/Implementation/RefService.cs
--- a/Engine/Source/Programs/UnrealCloudDDC/Jupiter/Implementation/RefService.cs
+++ b/Engine/Source/Programs/UnrealCloudDDC/Jupiter/Implementation/RefService.cs
@@ -3,6 +3,7 @@
 using System;
 using System.Collections.Generic;
 using System.Linq;
+using System.Threading;
 using System.Threading.Tasks;
 using EpicGames.AspNet;
 using EpicGames.Horde.Storage;
@@ -47,13 +48,13 @@ namespace Jupiter.Implementation
 			_cloudDDCSettings = cloudDDCSettings;
 		}
 
-		public Task<(RefRecord, BlobContents?)> GetAsync(NamespaceId ns, BucketId bucket, RefId key, string[]? fields = null, bool doLastAccessTracking = true)
+		public Task<(RefRecord, BlobContents?)> GetAsync(NamespaceId ns, BucketId bucket, RefId key, string[]? fields = null, bool doLastAccessTracking = true, CancellationToken cancellationToken = default)
 		{
-			return GetAsync(ns, bucket, key, fields, doLastAccessTracking, skipCache: false);
+			return GetAsync(ns, bucket, key, fields, doLastAccessTracking, skipCache: false, cancellationToken: cancellationToken);
 		}
 
 		// ReSharper disable once MethodOverloadWithOptionalParameter - this private overload exists only for bypassing cache for internal use within this service
-		private async Task<(RefRecord, BlobContents?)> GetAsync(NamespaceId ns, BucketId bucket, RefId key, string[]? fields = null, bool doLastAccessTracking = true, bool skipCache = false)
+		private async Task<(RefRecord, BlobContents?)> GetAsync(NamespaceId ns, BucketId bucket, RefId key, string[]? fields = null, bool doLastAccessTracking = true, bool skipCache = false, CancellationToken cancellationToken = default)
 		{
 			// if no field filtering is being used we assume everything is needed
 			IReferencesStore.FieldFlags flags = IReferencesStore.FieldFlags.All;
@@ -84,7 +85,7 @@ namespace Jupiter.Implementation
 			{
 				using ServerTimingMetricScoped? serverTimingScope = serverTiming?.CreateServerTimingMetricScope("ref.get", "Fetching Ref from DB");
 
-				o = await _referencesStore.GetAsync(ns, bucket, key, flags, opFlags);
+				o = await _referencesStore.GetAsync(ns, bucket, key, flags, opFlags, cancellationToken);
 			}
 
 			if (doLastAccessTracking)
@@ -116,14 +117,14 @@ namespace Jupiter.Implementation
 				{
 					using ServerTimingMetricScoped? serverTimingScope = serverTiming?.CreateServerTimingMetricScope("blob.get", "Downloading blob from store");
 
-					blobContents = await _blobService.GetObjectAsync(ns, o.BlobIdentifier);
+					blobContents = await _blobService.GetObjectAsync(ns, o.BlobIdentifier, cancellationToken: cancellationToken);
 				}
 			}
 
 			return (o, blobContents);
 		}
 
-		public async Task<(ContentId[], BlobId[])> PutAsync(NamespaceId ns, BucketId bucket, RefId key, BlobId blobHash, CbObject payload)
+		public async Task<(ContentId[], BlobId[])> PutAsync(NamespaceId ns, BucketId bucket, RefId key, BlobId blobHash, CbObject payload, CancellationToken cancellationToken = default)
 		{
 			IServerTiming? serverTiming = _httpContextAccessor.HttpContext?.RequestServices.GetService<IServerTiming>();
 			using ServerTimingMetricScoped? serverTimingScope = serverTiming?.CreateServerTimingMetricScope("ref.put", "Inserting ref");
@@ -133,12 +134,12 @@ namespace Jupiter.Implementation
 			// if we have no references we are always finalized, e.g. there are no referenced blobs to upload
 			bool isFinalized = !hasReferences;
 
-			Task objectStorePut = _referencesStore.PutAsync(ns, bucket, key, blobHash, payload.GetView().ToArray(), isFinalized);
+			Task objectStorePut = _referencesStore.PutAsync(ns, bucket, key, blobHash, payload.GetView().ToArray(), isFinalized, cancellationToken);
 
 			Task<BlobId>? blobStorePut = null;
 			if (_cloudDDCSettings.CurrentValue.EnablePutRefBodyIntoBlobStore)
 			{
-				blobStorePut = _blobService.PutObjectAsync(ns, payload.GetView().ToArray(), blobHash);
+				blobStorePut = _blobService.PutObjectAsync(ns, payload.GetView().ToArray(), blobHash, cancellationToken);
 			}
 
 			await objectStorePut;
@@ -147,7 +148,7 @@ namespace Jupiter.Implementation
 				await blobStorePut;
 			}
 
-			return await DoFinalizeAsync(ns, bucket, key, blobHash, payload);
+			return await DoFinalizeAsync(ns, bucket, key, blobHash, payload, cancellationToken);
 		}
 
 		private bool HasAttachments(CbObject payload)
@@ -181,16 +182,16 @@ namespace Jupiter.Implementation
 			return payload.Any(FieldHasAttachments);
 		}
 
-		public async Task<(ContentId[], BlobId[])> FinalizeAsync(NamespaceId ns, BucketId bucket, RefId key, BlobId blobHash)
+		public async Task<(ContentId[], BlobId[])> FinalizeAsync(NamespaceId ns, BucketId bucket, RefId key, BlobId blobHash, CancellationToken cancellationToken = default)
 		{
 			// finalize is intended to verify the state of the object in the db, so we bypass any caches to make sure the object we are working on is not stale
-			(RefRecord o, BlobContents? blob) = await GetAsync(ns, bucket, key, skipCache: true);
+			(RefRecord o, BlobContents? blob) = await GetAsync(ns, bucket, key, skipCache: true, cancellationToken: cancellationToken);
 			if (blob == null)
 			{
 				throw new InvalidOperationException("No blob when attempting to finalize");
 			}
 
-			byte[] blobContents = await blob.Stream.ToByteArrayAsync();
+			byte[] blobContents = await blob.Stream.ToByteArrayAsync(cancellationToken);
 			CbObject payload = new CbObject(blobContents);
 
 			if (!o.BlobIdentifier.Equals(blobHash))
@@ -198,17 +199,17 @@ namespace Jupiter.Implementation
 				throw new ObjectHashMismatchException(ns, bucket, key, blobHash, o.BlobIdentifier);
 			}
 
-			return await DoFinalizeAsync(ns, bucket, key, blobHash, payload);
+			return await DoFinalizeAsync(ns, bucket, key, blobHash, payload, cancellationToken);
 		}
 
-		private async Task<(ContentId[], BlobId[])> DoFinalizeAsync(NamespaceId ns, BucketId bucket, RefId key, BlobId blobHash, CbObject payload)
+		private async Task<(ContentId[], BlobId[])> DoFinalizeAsync(NamespaceId ns, BucketId bucket, RefId key, BlobId blobHash, CbObject payload, CancellationToken cancellationToken = default)
 		{
 			IServerTiming? serverTiming = _httpContextAccessor.HttpContext?.RequestServices.GetService<IServerTiming>();
 			using ServerTimingMetricScoped? serverTimingScope = serverTiming?.CreateServerTimingMetricScope("ref.finalize", "Finalizing the ref");
 
-			Task addRefToBlobsTask = _blobIndex.AddRefToBlobsAsync(ns, bucket, key, new[] { blobHash });
+			Task addRefToBlobsTask = _blobIndex.AddRefToBlobsAsync(ns, bucket, key, new[] { blobHash }, cancellationToken);
 			Task addToBucketListTask = _cloudDDCSettings.CurrentValue.EnableBucketStatsTracking
-				? _blobIndex.AddBlobToBucketListAsync(ns, bucket, key, blobHash, (long)payload.GetView().Length)
+				? _blobIndex.AddBlobToBucketListAsync(ns, bucket, key, blobHash, (long)payload.GetView().Length, cancellationToken)
 				: Task.CompletedTask;
 			ContentId[] missingReferences = Array.Empty<ContentId>();
 			BlobId[] missingBlobs = Array.Empty<BlobId>();
@@ -222,7 +223,7 @@ namespace Jupiter.Implementation
 				using TelemetrySpan _ = _tracer.StartActiveSpan("ObjectService.ResolveReferences").SetAttribute("operation.name", "ObjectService.ResolveReferences");
 				try
 				{
-					IAsyncEnumerable<BlobId> references = _referenceResolver.GetReferencedBlobsAsync(ns, payload);
+					IAsyncEnumerable<BlobId> references = _referenceResolver.GetReferencedBlobsAsync(ns, payload, cancellationToken: cancellationToken);
 
 					await foreach (BlobId blobId in references)
 					{
@@ -233,16 +234,16 @@ namespace Jupiter.Implementation
 								// if a blob is missing its not a error, the finalize will report this as missing and it will be uploaded and finalize ran again
 								try
 								{
-									BlobMetadata result = await _blobService.GetObjectMetadataAsync(ns, blobId);
-									await _blobIndex.AddBlobToBucketListAsync(ns, bucket, key, blobId, result.Length);
+									BlobMetadata result = await _blobService.GetObjectMetadataAsync(ns, blobId, cancellationToken);
+									await _blobIndex.AddBlobToBucketListAsync(ns, bucket, key, blobId, result.Length, cancellationToken);
 								}
 								catch (BlobNotFoundException)
 								{
 								}
-							}));
+							}, cancellationToken));
 						}
 
-						addRefMappingTasks.Add(_blobIndex.AddRefToBlobsAsync(ns, bucket, key, new BlobId[] { blobId }));
+						addRefMappingTasks.Add(_blobIndex.AddRefToBlobsAsync(ns, bucket, key, new BlobId[] { blobId }, cancellationToken));
 					}
 				}
 				catch (PartialReferenceResolveException e)
@@ -259,7 +260,7 @@ namespace Jupiter.Implementation
 
 			if (missingReferences.Length == 0 && missingBlobs.Length == 0)
 			{
-				await _referencesStore.FinalizeAsync(ns, bucket, key, blobHash);
+				await _referencesStore.FinalizeAsync(ns, bucket, key, blobHash, cancellationToken);
 				await _replicationLog.InsertAddEventAsync(ns, bucket, key, blobHash);
 			}
 
@@ -268,31 +269,31 @@ namespace Jupiter.Implementation
 			return (missingReferences, missingBlobs);
 		}
 
-		public IAsyncEnumerable<NamespaceId> GetNamespacesAsync()
+		public IAsyncEnumerable<NamespaceId> GetNamespacesAsync(CancellationToken cancellationToken)
 		{
-			return _referencesStore.GetNamespacesAsync();
+			return _referencesStore.GetNamespacesAsync(cancellationToken);
 		}
 
-		public Task<bool> DeleteAsync(NamespaceId ns, BucketId bucket, RefId key)
+		public Task<bool> DeleteAsync(NamespaceId ns, BucketId bucket, RefId key, CancellationToken cancellationToken)
 		{
-			return _referencesStore.DeleteAsync(ns, bucket, key);
+			return _referencesStore.DeleteAsync(ns, bucket, key, cancellationToken);
 		}
 
-		public Task<long> DropNamespaceAsync(NamespaceId ns)
+		public Task<long> DropNamespaceAsync(NamespaceId ns, CancellationToken cancellationToken)
 		{
-			return _referencesStore.DropNamespaceAsync(ns);
+			return _referencesStore.DropNamespaceAsync(ns, cancellationToken);
 		}
 
-		public Task<long> DeleteBucketAsync(NamespaceId ns, BucketId bucket)
+		public Task<long> DeleteBucketAsync(NamespaceId ns, BucketId bucket, CancellationToken cancellationToken)
 		{
-			return _referencesStore.DeleteBucketAsync(ns, bucket);
+			return _referencesStore.DeleteBucketAsync(ns, bucket, cancellationToken);
 		}
 
-		public async Task<bool> ExistsAsync(NamespaceId ns, BucketId bucket, RefId key)
+		public async Task<bool> ExistsAsync(NamespaceId ns, BucketId bucket, RefId key, CancellationToken cancellationToken)
 		{
 			try
 			{
-				(RefRecord, BlobContents?) _ = await GetAsync(ns, bucket, key, new string[] { "name" }, doLastAccessTracking: false, skipCache: true);
+				(RefRecord, BlobContents?) _ = await GetAsync(ns, bucket, key, new string[] { "name" }, doLastAccessTracking: false, skipCache: true, cancellationToken: cancellationToken);
 			}
 			catch (NamespaceNotFoundException)
 			{
@@ -318,23 +319,23 @@ namespace Jupiter.Implementation
 			return true;
 		}
 
-		public async Task<List<BlobId>> GetReferencedBlobsAsync(NamespaceId ns, BucketId bucket, RefId name, bool ignoreMissingBlobs = false)
+		public async Task<List<BlobId>> GetReferencedBlobsAsync(NamespaceId ns, BucketId bucket, RefId name, bool ignoreMissingBlobs = false, CancellationToken cancellationToken = default)
 		{
 			byte[] blob;
-			RefRecord o = await _referencesStore.GetAsync(ns, bucket, name, IReferencesStore.FieldFlags.IncludePayload, IReferencesStore.OperationFlags.None);
+			RefRecord o = await _referencesStore.GetAsync(ns, bucket, name, IReferencesStore.FieldFlags.IncludePayload, IReferencesStore.OperationFlags.None, cancellationToken);
 			if (o.InlinePayload != null && o.InlinePayload.Length != 0)
 			{
 				blob = o.InlinePayload;
 			}
 			else
 			{
-				BlobContents blobContents = await _blobService.GetObjectAsync(ns, o.BlobIdentifier);
-				blob = await blobContents.Stream.ToByteArrayAsync();
+				BlobContents blobContents = await _blobService.GetObjectAsync(ns, o.BlobIdentifier, cancellationToken: cancellationToken);
+				blob = await blobContents.Stream.ToByteArrayAsync(cancellationToken);
 			}
 
 			CbObject cbObject = new CbObject(blob);
 
-			List<BlobId> referencedBlobs = await _referenceResolver.GetReferencedBlobsAsync(ns, cbObject, ignoreMissingBlobs).ToListAsync();
+			List<BlobId> referencedBlobs = await _referenceResolver.GetReferencedBlobsAsync(ns, cbObject, ignoreMissingBlobs, cancellationToken).ToListAsync(cancellationToken);
 			return referencedBlobs;
 		}
 	}
diff --git a/Engine/Source/Programs/UnrealCloudDDC/Jupiter/Implementation/Replication/RefsReplicator.cs b/Engine/Source/Programs/UnrealCloudDDC/Jupiter/Implementation/Replication/RefsReplicator.cs
--- a/Engine/Source/Programs/UnrealCloudDDC/Jupiter/Implementation/Replication/RefsReplicator.cs
+++ b/Engine/Source/Programs/UnrealCloudDDC/Jupiter/Implementation/Replication/RefsReplicator.cs
@@ -541,7 +541,7 @@ namespace Jupiter.Implementation
 			Array.Copy(refs.References, potentialBlobs, refs.References.Length);
 			potentialBlobs[^1] = objectToReplicate;
 
-			BlobId[] missingBlobs = await _blobService.FilterOutKnownBlobsAsync(ns, potentialBlobs);
+			BlobId[] missingBlobs = await _blobService.FilterOutKnownBlobsAsync(ns, potentialBlobs, cancellationToken);
 			Task[] blobReplicationTasks = new Task[missingBlobs.Length];
 			for (int i = 0; i < missingBlobs.Length; i++)
 			{
@@ -597,9 +597,9 @@ namespace Jupiter.Implementation
 						throw new Exception("Expected content-length on blob response");
 					}
 
-					using IBufferedPayload payload = await _bufferedPayloadFactory.CreateFromStreamAsync(s, contentLength.Value);
+					using IBufferedPayload payload = await _bufferedPayloadFactory.CreateFromStreamAsync(s, contentLength.Value, cancellationToken);
 
-					await _blobService.PutObjectAsync(ns, payload, blobToReplicate);
+					await _blobService.PutObjectAsync(ns, payload, blobToReplicate, cancellationToken);
 				}, cancellationToken);
 			}
 
diff --git a/Engine/Source/Programs/UnrealCloudDDC/Jupiter/Implementation/TransactionLog/ReplicationLogSnapshot.cs b/Engine/Source/Programs/UnrealCloudDDC/Jupiter/Implementation/TransactionLog/ReplicationLogSnapshot.cs
--- a/Engine/Source/Programs/UnrealCloudDDC/Jupiter/Implementation/TransactionLog/ReplicationLogSnapshot.cs
+++ b/Engine/Source/Programs/UnrealCloudDDC/Jupiter/Implementation/TransactionLog/ReplicationLogSnapshot.cs
@@ -6,6 +6,7 @@ using System.IO;
 using System.IO.Compression;
 using System.Text;
 using System.Text.Json;
+using System.Threading;
 using EpicGames.Horde.Storage;
 using JetBrains.Annotations;
 using Jupiter.Common.Implementation;
@@ -76,7 +77,7 @@ namespace Jupiter.Implementation.TransactionLog
 
 		public static ReplicationLogSnapshot FromStream(BufferedPayloadFactory payloadFactory, Stream stream)
 		{
-			IBufferedPayload payload = payloadFactory.CreateFilesystemBufferedPayloadAsync(stream).Result;
+			IBufferedPayload payload = payloadFactory.CreateFilesystemBufferedPayloadAsync(stream, CancellationToken.None).Result;
 
 			using Stream payloadStream = payload.GetStream();
 			(NamespaceId ns, string lastBucket, Guid lastEvent, ulong countOfObjects) = ReadHeader(payloadStream);
diff --git a/Engine/Source/Programs/UnrealCloudDDC/Jupiter/Implementation/TransactionLog/ReplicationLogSnapshotBuilder.cs b/Engine/Source/Programs/UnrealCloudDDC/Jupiter/Implementation/TransactionLog/ReplicationLogSnapshotBuilder.cs
--- a/Engine/Source/Programs/UnrealCloudDDC/Jupiter/Implementation/TransactionLog/ReplicationLogSnapshotBuilder.cs
+++ b/Engine/Source/Programs/UnrealCloudDDC/Jupiter/Implementation/TransactionLog/ReplicationLogSnapshotBuilder.cs
@@ -46,13 +46,13 @@ namespace Jupiter.Implementation.TransactionLog
 			if (snapshotInfo != null)
 			{
 				// append to the previous snapshot if one is available
-				await using BlobContents blobContents = await _blobService.GetObjectAsync(snapshotInfo.BlobNamespace, snapshotInfo.SnapshotBlob);
+				await using BlobContents blobContents = await _blobService.GetObjectAsync(snapshotInfo.BlobNamespace, snapshotInfo.SnapshotBlob, cancellationToken: cancellationToken);
 				if (cancellationToken.IsCancellationRequested)
 				{
 					throw new TaskCanceledException();
 				}
 
-				using IBufferedPayload snapshotPayload = await _bufferedPayloadFactory.CreateFilesystemBufferedPayloadAsync(blobContents.Stream);
+				using IBufferedPayload snapshotPayload = await _bufferedPayloadFactory.CreateFilesystemBufferedPayloadAsync(blobContents.Stream, cancellationToken);
 				await using Stream s = snapshotPayload.GetStream();
 				snapshot = _replicationLogFactory.DeserializeSnapshotFromStream(s);
 				lastBucket = snapshot.LastBucket;
@@ -89,7 +89,7 @@ namespace Jupiter.Implementation.TransactionLog
 			}
 
 			Stream tempFileStream = tempFile.OpenRead();
-			using IBufferedPayload payload = await _bufferedPayloadFactory.CreateFromStreamAsync(tempFileStream, tempFile.Length);
+			using IBufferedPayload payload = await _bufferedPayloadFactory.CreateFromStreamAsync(tempFileStream, tempFile.Length, cancellationToken);
 			tempFileStream.Close();
 			tempFile.Delete();
 
@@ -97,7 +97,7 @@ namespace Jupiter.Implementation.TransactionLog
 				BlobId blobIdentifier;
 				{
 					await using Stream stream = payload.GetStream();
-					blobIdentifier = await BlobId.FromStreamAsync(stream);
+					blobIdentifier = await BlobId.FromStreamAsync(stream, cancellationToken);
 				}
 
 				CbWriter writer = new CbWriter();
@@ -115,9 +115,9 @@ namespace Jupiter.Implementation.TransactionLog
 				}
 
 				// upload the attachment first so we are not missing any references when we go to create the ref
-				await _blobService.PutObjectAsync(storeInNamespace, payload, blobIdentifier);
+				await _blobService.PutObjectAsync(storeInNamespace, payload, blobIdentifier, cancellationToken);
 
-				(ContentId[] missingContentIds, BlobId[] missingBlobs) = await _refService.PutAsync(storeInNamespace, new BucketId("snapshot"), new RefId(blobIdentifier.ToString()), cbBlobId, new CbObject(cbObjectBytes));
+				(ContentId[] missingContentIds, BlobId[] missingBlobs) = await _refService.PutAsync(storeInNamespace, new BucketId("snapshot"), new RefId(blobIdentifier.ToString()), cbBlobId, new CbObject(cbObjectBytes), cancellationToken);
 				List<ContentHash> missingHashes = new List<ContentHash>(missingContentIds);
 				missingHashes.AddRange(missingBlobs);
 				if (missingHashes.Count != 0)
