Unreal Cloud DDC - Improved logging when running into issues with ondemand replication and made the filesystem GC resistent to directories getting removed while it is running

diff --git a/Engine/Source/Programs/UnrealCloudDDC/Jupiter/Implementation/Blob/BlobService.cs b/Engine/Source/Programs/UnrealCloudDDC/Jupiter/Implementation/Blob/BlobService.cs
--- a/Engine/Source/Programs/UnrealCloudDDC/Jupiter/Implementation/Blob/BlobService.cs
+++ b/Engine/Source/Programs/UnrealCloudDDC/Jupiter/Implementation/Blob/BlobService.cs
@@ -613,6 +613,7 @@ public class BlobService : IBlobService
 
 			if (blobResponse.StatusCode != HttpStatusCode.OK)
 			{
+				_logger.LogWarning("Failed to replicate {Blob} in {Namespace} from region {Region} due to bad http status code {StatusCode}.", blob, ns, region, blobResponse.StatusCode);
 				throw new BlobReplicationException(ns, blob, $"Failed to replicate {blob} in {ns} from region {region} due to bad http status code {blobResponse.StatusCode}.");
 			}
 
@@ -625,6 +626,8 @@ public class BlobService : IBlobService
 
 		if (!replicated)
 		{
+			_logger.LogWarning("Failed to replicate {Blob} in {Namespace} due to it not existing in any region", blob, ns);
+
 			throw new BlobReplicationException(ns, blob, $"Failed to replicate {blob} in {ns} due to it not existing in any region");
 		}
 
diff --git a/Engine/Source/Programs/UnrealCloudDDC/Jupiter/Implementation/Blob/FileSystemStore.cs b/Engine/Source/Programs/UnrealCloudDDC/Jupiter/Implementation/Blob/FileSystemStore.cs
--- a/Engine/Source/Programs/UnrealCloudDDC/Jupiter/Implementation/Blob/FileSystemStore.cs
+++ b/Engine/Source/Programs/UnrealCloudDDC/Jupiter/Implementation/Blob/FileSystemStore.cs
@@ -241,6 +241,10 @@ namespace Jupiter.Implementation
 						{
 							// if the file was gced while running we can just ignore it
 						}
+						catch (DirectoryNotFoundException)
+						{
+							// if the directory was deleted then we can ignore it to, end result is the blob is gone which is what matters
+						}
 					}
 				}
 
