Unreal Cloud DDC - Added option to restart pods when the config map changes

diff --git a/Engine/Source/Programs/UnrealCloudDDC/Helm/UnrealCloudDDC/templates/deployment-worker.yaml b/Engine/Source/Programs/UnrealCloudDDC/Helm/UnrealCloudDDC/templates/deployment-worker.yaml
--- a/Engine/Source/Programs/UnrealCloudDDC/Helm/UnrealCloudDDC/templates/deployment-worker.yaml
+++ b/Engine/Source/Programs/UnrealCloudDDC/Helm/UnrealCloudDDC/templates/deployment-worker.yaml
@@ -29,11 +29,14 @@ spec:
       {{- include "common.tplvalues.render" ( dict "value" $worker.podLabels "context" $ ) | nindent 8 }}
       {{- end }}
 
-      {{- if $worker.podAnnotations }}
+      {{- if or $worker.podAnnotations .Values.restartPodOnConfigMapChange }}
       annotations:
         {{- if $worker.podAnnotations }}
         {{- include "common.tplvalues.render" (dict "value" $worker.podAnnotations "context" $) | nindent 8 }}
         {{- end }}
+        {{- if .Values.restartPodOnConfigMapChange }}
+        checksum/config: {{ include (print $.Template.BasePath "/configmap-worker.yaml") . | sha256sum }}
+        {{- end }}
       {{- end }}
     spec:
       serviceAccountName: {{ include "jupiter.serviceAccountName" . }}
diff --git a/Engine/Source/Programs/UnrealCloudDDC/Helm/UnrealCloudDDC/templates/deployment.yaml b/Engine/Source/Programs/UnrealCloudDDC/Helm/UnrealCloudDDC/templates/deployment.yaml
--- a/Engine/Source/Programs/UnrealCloudDDC/Helm/UnrealCloudDDC/templates/deployment.yaml
+++ b/Engine/Source/Programs/UnrealCloudDDC/Helm/UnrealCloudDDC/templates/deployment.yaml
@@ -15,11 +15,14 @@ spec:
       {{- include "common.tplvalues.render" ( dict "value" .Values.podLabels "context" $ ) | nindent 8 }}
       {{- end }}
 
-      {{- if .Values.podAnnotations }}
+      {{- if or $worker.podAnnotations .Values.restartPodOnConfigMapChange }}
       annotations:
         {{- if .Values.podAnnotations }}
         {{- include "common.tplvalues.render" (dict "value" .Values.podAnnotations "context" $) | nindent 8 }}
         {{- end }}
+        {{- if .Values.restartPodOnConfigMapChange }}
+        checksum/config: {{ include (print $.Template.BasePath "/configmap.yaml") . | sha256sum }}
+        {{- end }}
       {{- end }}
     spec:
       serviceAccountName: {{ include "jupiter.serviceAccountName" . }}
@@ -164,7 +167,7 @@ spec:
     {{- if .Values.nginx.enabled }}
       - name: nginx
         image: nginx:1.25.4
-        {{- if .Values.nginx.env }}
+        {{- if .Values.nginx.env }} 
         env:
         {{- include "common.tplvalues.render" ( dict "value" .Values.nginx.env "context" $ ) | nindent 8 }}
         {{- end }}
diff --git a/Engine/Source/Programs/UnrealCloudDDC/Helm/UnrealCloudDDC/values.yaml b/Engine/Source/Programs/UnrealCloudDDC/Helm/UnrealCloudDDC/values.yaml
--- a/Engine/Source/Programs/UnrealCloudDDC/Helm/UnrealCloudDDC/values.yaml
+++ b/Engine/Source/Programs/UnrealCloudDDC/Helm/UnrealCloudDDC/values.yaml
@@ -307,6 +307,11 @@ useDotnetMonitorSidecar : false
 ## this is typically faster then falling back to global blob store
 fetchBlobsFromPeerPods: false
 
+## restartPodOnConfigMapChange
+## Adds a annotation that will include the checksum of the config map to all pods, causing them to restart if the configmap changes
+## This is a bit pessimistic as config changes can be dynamically reloaded but its inconsistent
+restartPodOnConfigMapChange: false
+
 ## Nginx configuration
 nginx:
 
