Unreal Cloud DDC - Added ability to define a hardcoded auth token for service to service communication, useful if using service accounts instead of OIDC. Can use secret strings.

diff --git a/Engine/Source/Programs/UnrealCloudDDC/Jupiter/Implementation/Replication/ReplicationSettings.cs b/Engine/Source/Programs/UnrealCloudDDC/Jupiter/Implementation/Replication/ReplicationSettings.cs
--- a/Engine/Source/Programs/UnrealCloudDDC/Jupiter/Implementation/Replication/ReplicationSettings.cs
+++ b/Engine/Source/Programs/UnrealCloudDDC/Jupiter/Implementation/Replication/ReplicationSettings.cs
@@ -110,5 +110,10 @@ namespace Jupiter
 		/// The authentication scheme to use for this token
 		/// </summary>
 		public string SchemeName { get; set; } = "Bearer";
+
+		/// <summary>
+		/// The access token to use if not using OIDC. Can be resolved using a secret string.
+		/// </summary>
+		public string? AccessToken { get; set; } = null;
 	}
 }
diff --git a/Engine/Source/Programs/UnrealCloudDDC/Jupiter/Implementation/ServiceCredentials.cs b/Engine/Source/Programs/UnrealCloudDDC/Jupiter/Implementation/ServiceCredentials.cs
--- a/Engine/Source/Programs/UnrealCloudDDC/Jupiter/Implementation/ServiceCredentials.cs
+++ b/Engine/Source/Programs/UnrealCloudDDC/Jupiter/Implementation/ServiceCredentials.cs
@@ -19,10 +19,12 @@ namespace Jupiter.Implementation
 	{
 		private readonly ClientCredentialOAuthAuthenticator? _authenticator;
 		private readonly IOptionsMonitor<ServiceCredentialSettings> _settings;
+		private readonly ISecretResolver _secretResolver;
 
 		public ServiceCredentials(IServiceProvider provider, IOptionsMonitor<ServiceCredentialSettings> settings, ISecretResolver secretResolver)
 		{
 			_settings = settings;
+			_secretResolver = secretResolver;
 			if (settings.CurrentValue.OAuthLoginUrl != null)
 			{
 				string? clientId = secretResolver.Resolve(settings.CurrentValue.OAuthClientId);
@@ -43,6 +45,10 @@ namespace Jupiter.Implementation
 
 		public string? GetToken()
 		{
+			if (_authenticator == null && !string.IsNullOrEmpty(_settings.CurrentValue.AccessToken))
+			{
+				return _secretResolver.Resolve(_settings.CurrentValue.AccessToken);
+			}
 			return _authenticator?.AuthenticateAsync().Result;
 		}
 
@@ -50,6 +56,11 @@ namespace Jupiter.Implementation
 		{
 			if (_authenticator == null)
 			{
+				if (!string.IsNullOrEmpty(_settings.CurrentValue.AccessToken))
+				{
+					return _secretResolver.Resolve(_settings.CurrentValue.AccessToken);
+				}
+
 				return null;
 			}
 			return await _authenticator.AuthenticateAsync();
