Unreal Cloud DDC - Fixed build errors in references controller

diff --git a/Engine/Source/Programs/UnrealCloudDDC/Jupiter/Controllers/ReferencesController.cs b/Engine/Source/Programs/UnrealCloudDDC/Jupiter/Controllers/ReferencesController.cs
--- a/Engine/Source/Programs/UnrealCloudDDC/Jupiter/Controllers/ReferencesController.cs
+++ b/Engine/Source/Programs/UnrealCloudDDC/Jupiter/Controllers/ReferencesController.cs
@@ -297,7 +297,7 @@ namespace Jupiter.Controllers
 							byte[] packageBytes;
 							{
 								using TelemetrySpan _ = _tracer.StartActiveSpan("cbpackage.buffer").SetAttribute("operation.name", "cbpackage.buffer");
-								packageBytes = await writer.ToByteArray();
+								packageBytes = await writer.ToByteArrayAsync();
 							}
 							await using BlobContents contents = new BlobContents(packageBytes);
 							await WriteBody(contents, CustomMediaTypeNames.UnrealCompactBinaryPackage);
@@ -744,11 +744,11 @@ namespace Jupiter.Controllers
 			_diagnosticContext.Set("Content-Length", Request.ContentLength ?? -1);
 
 			byte[] b = await RequestUtil.ReadRawBodyAsync(Request);
-			CbPackageReader packageReader = await CbPackageReader.Create(new MemoryStream(b));
+			CbPackageReader packageReader = await CbPackageReader.CreateAsync(new MemoryStream(b));
 
 			try
 			{
-				await foreach ((CbPackageAttachmentEntry entry, byte[] blob) in packageReader.IterateAttachments())
+				await foreach ((CbPackageAttachmentEntry entry, byte[] blob) in packageReader.IterateAttachmentsAsync())
 				{
 					if (entry.Flags.HasFlag(CbPackageAttachmentFlags.IsError))
 					{
