Unreal Cloud DDC - Fixed incorrect yaml when attempting to add podAnnotations for automatic restarts of pods

diff --git a/Engine/Source/Programs/UnrealCloudDDC/Helm/UnrealCloudDDC/templates/deployment.yaml b/Engine/Source/Programs/UnrealCloudDDC/Helm/UnrealCloudDDC/templates/deployment.yaml
--- a/Engine/Source/Programs/UnrealCloudDDC/Helm/UnrealCloudDDC/templates/deployment.yaml
+++ b/Engine/Source/Programs/UnrealCloudDDC/Helm/UnrealCloudDDC/templates/deployment.yaml
@@ -15,7 +15,7 @@ spec:
       {{- include "common.tplvalues.render" ( dict "value" .Values.podLabels "context" $ ) | nindent 8 }}
       {{- end }}
 
-      {{- if or $worker.podAnnotations .Values.restartPodOnConfigMapChange }}
+      {{- if or .Values.podAnnotations .Values.restartPodOnConfigMapChange }}
       annotations:
         {{- if .Values.podAnnotations }}
         {{- include "common.tplvalues.render" (dict "value" .Values.podAnnotations "context" $) | nindent 8 }}
