Unreal Cloud DDC - Avoiding preparing bucket statistics queries for non scylla cassandra databases as it causes issues for CosmosDB

diff --git a/Engine/Source/Programs/UnrealCloudDDC/Jupiter/Implementation/Blob/ScyllaBlobIndex.cs b/Engine/Source/Programs/UnrealCloudDDC/Jupiter/Implementation/Blob/ScyllaBlobIndex.cs
--- a/Engine/Source/Programs/UnrealCloudDDC/Jupiter/Implementation/Blob/ScyllaBlobIndex.cs
+++ b/Engine/Source/Programs/UnrealCloudDDC/Jupiter/Implementation/Blob/ScyllaBlobIndex.cs
@@ -25,8 +25,8 @@ public class ScyllaBlobIndex : IBlobIndex
 	private readonly Tracer _tracer;
 	private readonly ISession _session;
 	private readonly Mapper _mapper;
-	private readonly PreparedStatement _getBucketStatsStatement;
-	private readonly PreparedStatement _getBucketStatsRefsStatement;
+	private readonly PreparedStatement? _getBucketStatsStatement;
+	private readonly PreparedStatement? _getBucketStatsRefsStatement;
 
 	public ScyllaBlobIndex(IScyllaSessionManager scyllaSessionManager, IOptionsMonitor<JupiterSettings> jupiterSettings, IOptionsMonitor<ScyllaSettings> scyllaSettings, INamespacePolicyResolver namespacePolicyResolver, Tracer tracer)
 	{
@@ -88,8 +88,11 @@ public class ScyllaBlobIndex : IBlobIndex
 			));
 		}
 
-		_getBucketStatsStatement = _session.Prepare("select count(blob_id), min(size), max(size), sum(size) from bucket_referenced_blobs WHERE namespace = ? AND bucket_id = ?  AND hash_prefix = ?");
-		_getBucketStatsRefsStatement = _session.Prepare("select count(reference_id) from bucket_referenced_ref WHERE namespace = ? AND bucket_id = ?  AND hash_prefix = ?");
+		if (scyllaSessionManager.IsScylla)
+		{
+			_getBucketStatsStatement = _session.Prepare("select count(blob_id), min(size), max(size), sum(size) from bucket_referenced_blobs WHERE namespace = ? AND bucket_id = ?  AND hash_prefix = ?");
+			_getBucketStatsRefsStatement = _session.Prepare("select count(reference_id) from bucket_referenced_ref WHERE namespace = ? AND bucket_id = ?  AND hash_prefix = ?");
+		}
 	}
 
 	public async Task AddBlobToIndexAsync(NamespaceId ns, BlobId id, string? region = null, CancellationToken cancellationToken = default)
@@ -391,6 +394,11 @@ public class ScyllaBlobIndex : IBlobIndex
 
 		Debug.Assert(i == 65536);
 
+		if (_getBucketStatsRefsStatement == null || _getBucketStatsStatement == null)
+		{
+			throw new Exception("Calculating bucket statistics is not supported when not using Scylla");
+		}
+
 		const int DegreeOfParallelism = 32;
 		await Parallel.ForEachAsync(hashPrefixes, new ParallelOptions { MaxDegreeOfParallelism = DegreeOfParallelism }, async (hashPrefix, token) =>
 		{
