Unreal Cloud DDC - Fixed issue with blob GC not correctly checking all references when determining if a blob needs to be gced

diff --git a/Engine/Source/Programs/UnrealCloudDDC/Jupiter/Implementation/GC/OrphanBlobCleanupRefs.cs b/Engine/Source/Programs/UnrealCloudDDC/Jupiter/Implementation/GC/OrphanBlobCleanupRefs.cs
--- a/Engine/Source/Programs/UnrealCloudDDC/Jupiter/Implementation/GC/OrphanBlobCleanupRefs.cs
+++ b/Engine/Source/Programs/UnrealCloudDDC/Jupiter/Implementation/GC/OrphanBlobCleanupRefs.cs
@@ -160,8 +160,13 @@ namespace Jupiter.Implementation
 
 						try
 						{
-							found = await _refService.ExistsAsync(blobNamespace, bucket, key, cancellationToken);
-							break;
+							bool refFound = await _refService.ExistsAsync(blobNamespace, bucket, key, cancellationToken);
+
+							if (refFound)
+							{
+								found = true;
+								break;
+							}
 						}
 						catch (RefNotFoundException)
 						{
