Unreal Cloud DDC - Upgraded Azure.Identity to resolve vulnerability warning

diff --git a/Engine/Source/Programs/UnrealCloudDDC/Jupiter.Common/Jupiter.Common.csproj b/Engine/Source/Programs/UnrealCloudDDC/Jupiter.Common/Jupiter.Common.csproj
--- a/Engine/Source/Programs/UnrealCloudDDC/Jupiter.Common/Jupiter.Common.csproj
+++ b/Engine/Source/Programs/UnrealCloudDDC/Jupiter.Common/Jupiter.Common.csproj
@@ -29,7 +29,7 @@
     <PackageReference Include="AWSSDK.Extensions.NETCore.Setup" Version="3.7.300" />
     <PackageReference Include="AWSSDK.SecretsManager" Version="3.7.302.8" />
     <PackageReference Include="AWSSDK.SecurityToken" Version="3.7.300.34" />
-    <PackageReference Include="Azure.Identity" Version="1.10.4" />
+    <PackageReference Include="Azure.Identity" Version="1.11.0" />
     <PackageReference Include="Azure.Security.KeyVault.Secrets" Version="4.5.0" />
     <PackageReference Include="Newtonsoft.Json" Version="13.0.3" />
     <PackageReference Include="OpenTelemetry" Version="1.7.0" />
diff --git a/Engine/Source/Programs/UnrealCloudDDC/Jupiter/Jupiter.csproj b/Engine/Source/Programs/UnrealCloudDDC/Jupiter/Jupiter.csproj
--- a/Engine/Source/Programs/UnrealCloudDDC/Jupiter/Jupiter.csproj
+++ b/Engine/Source/Programs/UnrealCloudDDC/Jupiter/Jupiter.csproj
@@ -72,7 +72,7 @@
     <PackageReference Include="AspNetCore.HealthChecks.System" Version="8.0.0" />
     <PackageReference Include="AWSSDK.Extensions.NETCore.Setup" Version="3.7.300" />
     <PackageReference Include="AWSSDK.S3" Version="3.7.305.8" />
-    <PackageReference Include="Azure.Identity" Version="1.10.4" />
+    <PackageReference Include="Azure.Identity" Version="1.11.0" />
     <PackageReference Include="Azure.Security.KeyVault.Secrets" Version="4.5.0" />
     <PackageReference Include="Azure.Storage.Blobs" Version="12.15.1" />
     <PackageReference Include="KubernetesClient" Version="13.0.1" />
