Unreal Cloud DDC - Prepared release 1.1.0

diff --git a/Engine/Source/Programs/UnrealCloudDDC/Helm/UnrealCloudDDC/version.yaml b/Engine/Source/Programs/UnrealCloudDDC/Helm/UnrealCloudDDC/version.yaml
--- a/Engine/Source/Programs/UnrealCloudDDC/Helm/UnrealCloudDDC/version.yaml
+++ b/Engine/Source/Programs/UnrealCloudDDC/Helm/UnrealCloudDDC/version.yaml
@@ -1 +1 @@
-version: 1.0.0
\ No newline at end of file
+version: 1.1.0
\ No newline at end of file
diff --git a/Engine/Source/Programs/UnrealCloudDDC/changelog.md b/Engine/Source/Programs/UnrealCloudDDC/changelog.md
--- a/Engine/Source/Programs/UnrealCloudDDC/changelog.md
+++ b/Engine/Source/Programs/UnrealCloudDDC/changelog.md
@@ -1,5 +1,14 @@
 # Unreleased
 
+# 1.1.0
+* Added ability to specify a access token if using ServiceAccounts to the ServiceCredentials (allowing for replication to function when no OIDC configuration is present)
+* Fixed issue with blob stats using statements which did not work for CosmosDB by disabling them which means blob stats will not function for CosmosDB (this feature is not enabled by default anyway)
+* Added ability to define a pod annotation that contains a checksum of the configmap, causing pods to get restarted when the config changes.
+* Fixed issues with batch endpoints for compressed blob endpoints, not used by any production workloads yet.
+* Speed up S3 blob listing for GC by listing them per S3 prefix. Slower for very small datasets were speed doesn't matter. Can be disabled with `S3.PerPrefixListing`
+* Changed filesystem cleanup to just delete objects that are quite old rather then finding the oldest, makes it much more responsive and needs less state in memory at the cost of the cleanup being a bit more random but this is just the filecache anyway.
+* Fixes to ondemand replication to avoid infinite recursion.
+
 # 1.0.0
 * .NET 8 Upgrade.
 * Fixes for very large payloads (2GB+)
diff --git a/Engine/Source/Programs/UnrealCloudDDC/version.yaml b/Engine/Source/Programs/UnrealCloudDDC/version.yaml
--- a/Engine/Source/Programs/UnrealCloudDDC/version.yaml
+++ b/Engine/Source/Programs/UnrealCloudDDC/version.yaml
@@ -1 +1 @@
-version: 1.0.0
\ No newline at end of file
+version: 1.1.0
\ No newline at end of file
