Unreal Cloud DDC - Upgraded open telemetry libraries to resolve vulnerability warnings

diff --git a/Engine/Source/Programs/UnrealCloudDDC/Jupiter.Common/Jupiter.Common.csproj b/Engine/Source/Programs/UnrealCloudDDC/Jupiter.Common/Jupiter.Common.csproj
--- a/Engine/Source/Programs/UnrealCloudDDC/Jupiter.Common/Jupiter.Common.csproj
+++ b/Engine/Source/Programs/UnrealCloudDDC/Jupiter.Common/Jupiter.Common.csproj
@@ -35,8 +35,8 @@
     <PackageReference Include="OpenTelemetry" Version="1.7.0" />
     <PackageReference Include="OpenTelemetry.Exporter.Console" Version="1.7.0" />
     <PackageReference Include="OpenTelemetry.Exporter.OpenTelemetryProtocol" Version="1.7.0" />
-    <PackageReference Include="OpenTelemetry.Instrumentation.AspNetCore" Version="1.7.0" />
-    <PackageReference Include="OpenTelemetry.Instrumentation.Http" Version="1.7.0" />
+    <PackageReference Include="OpenTelemetry.Instrumentation.AspNetCore" Version="1.8.1" />
+    <PackageReference Include="OpenTelemetry.Instrumentation.Http" Version="1.8.1" />
     <PackageReference Include="OpenTelemetry.Instrumentation.Runtime" Version="1.7.0" />
     <PackageReference Include="OpenTelemetry.Extensions.Hosting" Version="1.7.0" />
     <PackageReference Include="IdentityServer4.AccessTokenValidation" Version="3.0.1" />
