Unreal Cloud DDC - Prepared version 1.1.1

diff --git a/Engine/Source/Programs/UnrealCloudDDC/Helm/UnrealCloudDDC/version.yaml b/Engine/Source/Programs/UnrealCloudDDC/Helm/UnrealCloudDDC/version.yaml
--- a/Engine/Source/Programs/UnrealCloudDDC/Helm/UnrealCloudDDC/version.yaml
+++ b/Engine/Source/Programs/UnrealCloudDDC/Helm/UnrealCloudDDC/version.yaml
@@ -1 +1 @@
-version: 1.1.0
\ No newline at end of file
+version: 1.1.1
\ No newline at end of file
diff --git a/Engine/Source/Programs/UnrealCloudDDC/changelog.md b/Engine/Source/Programs/UnrealCloudDDC/changelog.md
--- a/Engine/Source/Programs/UnrealCloudDDC/changelog.md
+++ b/Engine/Source/Programs/UnrealCloudDDC/changelog.md
@@ -1,5 +1,8 @@
 # Unreleased
 
+# 1.1.1
+* Fixes to helm chart when using ServiceAccounts for authentication and configuring replication
+
 # 1.1.0
 * Added ability to specify a access token if using ServiceAccounts to the ServiceCredentials (allowing for replication to function when no OIDC configuration is present)
 * Fixed issue with blob stats using statements which did not work for CosmosDB by disabling them which means blob stats will not function for CosmosDB (this feature is not enabled by default anyway)
diff --git a/Engine/Source/Programs/UnrealCloudDDC/version.yaml b/Engine/Source/Programs/UnrealCloudDDC/version.yaml
--- a/Engine/Source/Programs/UnrealCloudDDC/version.yaml
+++ b/Engine/Source/Programs/UnrealCloudDDC/version.yaml
@@ -1 +1 @@
-version: 1.1.0
\ No newline at end of file
+version: 1.1.1
\ No newline at end of file
