Unreal Cloud DDC - Added some traced properties to ref operations to be able to understand frequency of particular namespaces and buckets

diff --git a/Engine/Source/Programs/UnrealCloudDDC/Jupiter/Controllers/ReferencesController.cs b/Engine/Source/Programs/UnrealCloudDDC/Jupiter/Controllers/ReferencesController.cs
--- a/Engine/Source/Programs/UnrealCloudDDC/Jupiter/Controllers/ReferencesController.cs
+++ b/Engine/Source/Programs/UnrealCloudDDC/Jupiter/Controllers/ReferencesController.cs
@@ -119,6 +119,9 @@ namespace Jupiter.Controllers
 
 			try
 			{
+				Tracer.CurrentSpan.SetAttribute("bucket", bucket.ToString());
+				Tracer.CurrentSpan.SetAttribute("namespace", ns.ToString());
+
 				(RefRecord objectRecord, BlobContents? blob) = await _refService.GetAsync(ns, bucket, key, Array.Empty<string>());
 
 				if (blob == null)
@@ -163,6 +166,7 @@ namespace Jupiter.Controllers
 
 				string responseType = _formatResolver.GetResponseType(Request, format, CustomMediaTypeNames.UnrealCompactBinary);
 				Tracer.CurrentSpan.SetAttribute("response-type", responseType);
+
 				switch (responseType)
 				{
 					case CustomMediaTypeNames.UnrealCompactBinary:
@@ -623,6 +627,9 @@ namespace Jupiter.Controllers
 			}
 
 			_diagnosticContext.Set("Content-Length", Request.ContentLength ?? -1);
+			
+			Tracer.CurrentSpan.SetAttribute("bucket", bucket.ToString());
+			Tracer.CurrentSpan.SetAttribute("namespace", ns.ToString());
 
 			CbObject payloadObject;
 			BlobId blobHeader;
@@ -742,6 +749,8 @@ namespace Jupiter.Controllers
 			}
 
 			_diagnosticContext.Set("Content-Length", Request.ContentLength ?? -1);
+			Tracer.CurrentSpan.SetAttribute("bucket", bucket.ToString());
+			Tracer.CurrentSpan.SetAttribute("namespace", ns.ToString());
 
 			byte[] b = await RequestUtil.ReadRawBodyAsync(Request);
 			CbPackageReader packageReader = await CbPackageReader.CreateAsync(new MemoryStream(b));
@@ -803,6 +812,9 @@ namespace Jupiter.Controllers
 
 			try
 			{
+				Tracer.CurrentSpan.SetAttribute("bucket", bucket.ToString());
+				Tracer.CurrentSpan.SetAttribute("namespace", ns.ToString());
+
 				(ContentId[] missingReferences, BlobId[] missingBlobs) = await _refService.FinalizeAsync(ns, bucket, key, hash, HttpContext.RequestAborted);
 				List<ContentHash> missingHashes = new List<ContentHash>(missingReferences);
 				missingHashes.AddRange(missingBlobs);
