Unreal Cloud DDC - Fixed build error (missing using)

diff --git a/Engine/Source/Programs/UnrealCloudDDC/Jupiter/Implementation/Objects/ScyllaReferencesStore.cs b/Engine/Source/Programs/UnrealCloudDDC/Jupiter/Implementation/Objects/ScyllaReferencesStore.cs
--- a/Engine/Source/Programs/UnrealCloudDDC/Jupiter/Implementation/Objects/ScyllaReferencesStore.cs
+++ b/Engine/Source/Programs/UnrealCloudDDC/Jupiter/Implementation/Objects/ScyllaReferencesStore.cs
@@ -3,6 +3,7 @@
 using System;
 using System.Collections.Concurrent;
 using System.Collections.Generic;
+using System.Diagnostics;
 using System.Linq;
 using System.Runtime.CompilerServices;
 using System.Threading;
