Unreal Cloud DDC - Moved the inline blob setting to allow us to know if a blob will be inlined when we are deciding to submit this into the blob store or not, also renamed it to from EnablePutRefBodyIntoBlobStore to EnableInlineSmallBlobs which should be easier to understand what it does. The feature toggle to disable the extra blob into the blob store is now EnableForceSubmitRefBlobToBlobStore.

This will be a breaking change for users using a DB that does not support inlining of blobs (CosmosDB) as they need to change the option used.

diff --git a/Engine/Source/Programs/UnrealCloudDDC/Helm/UnrealCloudDDC/templates/configmap.yaml b/Engine/Source/Programs/UnrealCloudDDC/Helm/UnrealCloudDDC/templates/configmap.yaml
--- a/Engine/Source/Programs/UnrealCloudDDC/Helm/UnrealCloudDDC/templates/configmap.yaml
+++ b/Engine/Source/Programs/UnrealCloudDDC/Helm/UnrealCloudDDC/templates/configmap.yaml
@@ -41,8 +41,12 @@ UnrealCloudDDC:
   EnableBucketStatsTracking: {{ $config.EnableBucketStatsTracking }}
 {{- end }}
 
-{{- if hasKey $config "EnablePutRefBodyIntoBlobStore" }}
-  EnablePutRefBodyIntoBlobStore: {{ $config.EnablePutRefBodyIntoBlobStore }}
+{{- if hasKey $config "EnableForceSubmitRefBlobToBlobStore" }}
+  EnableForceSubmitRefBlobToBlobStore: {{ $config.EnableForceSubmitRefBlobToBlobStore }}
+{{- end }}
+
+{{- if hasKey $config "EnableInlineSmallBlobs" }}
+  EnableInlineSmallBlobs: {{ $config.EnableInlineSmallBlobs }}
 {{- end }}
 
 {{ if $config.StorageImplementations }}
diff --git a/Engine/Source/Programs/UnrealCloudDDC/Jupiter/Implementation/Objects/ScyllaReferencesStore.cs b/Engine/Source/Programs/UnrealCloudDDC/Jupiter/Implementation/Objects/ScyllaReferencesStore.cs
--- a/Engine/Source/Programs/UnrealCloudDDC/Jupiter/Implementation/Objects/ScyllaReferencesStore.cs
+++ b/Engine/Source/Programs/UnrealCloudDDC/Jupiter/Implementation/Objects/ScyllaReferencesStore.cs
@@ -145,12 +145,6 @@ namespace Jupiter.Implementation
 		{
 			using TelemetrySpan scope = _tracer.BuildScyllaSpan("scylla.put").SetAttribute("resource.name", $"{ns}.{bucket}.{name}");
 
-			if (blob.LongLength > _settings.CurrentValue.InlineBlobMaxSize)
-			{
-				// do not inline large blobs
-				blob = Array.Empty<byte>();
-			}
-
 			// add the bucket in parallel with inserting the actual object
 			Task addBucketTask = AddBucketAsync(ns, bucket);
 
diff --git a/Engine/Source/Programs/UnrealCloudDDC/Jupiter/Implementation/Objects/ScyllaSettings.cs b/Engine/Source/Programs/UnrealCloudDDC/Jupiter/Implementation/Objects/ScyllaSettings.cs
--- a/Engine/Source/Programs/UnrealCloudDDC/Jupiter/Implementation/Objects/ScyllaSettings.cs
+++ b/Engine/Source/Programs/UnrealCloudDDC/Jupiter/Implementation/Objects/ScyllaSettings.cs
@@ -9,7 +9,6 @@ namespace Jupiter.Implementation;
 public class ScyllaSettings : IValidatableObject
 {
 	public string ConnectionString { get; set; } = "Contact Points=localhost,scylla;Default Keyspace=jupiter";
-	public long InlineBlobMaxSize { get; set; } = 32 * 1024; // default to 32 kb blobs max
 	public int MaxSnapshotsPerNamespace { get; set; } = 10;
 
 	/// <summary>
diff --git a/Engine/Source/Programs/UnrealCloudDDC/Jupiter/Implementation/RefService.cs b/Engine/Source/Programs/UnrealCloudDDC/Jupiter/Implementation/RefService.cs
--- a/Engine/Source/Programs/UnrealCloudDDC/Jupiter/Implementation/RefService.cs
+++ b/Engine/Source/Programs/UnrealCloudDDC/Jupiter/Implementation/RefService.cs
@@ -134,10 +134,25 @@ namespace Jupiter.Implementation
 			// if we have no references we are always finalized, e.g. there are no referenced blobs to upload
 			bool isFinalized = !hasReferences;
 
-			Task objectStorePut = _referencesStore.PutAsync(ns, bucket, key, blobHash, payload.GetView().ToArray(), isFinalized, cancellationToken);
+			// if inlining is enabled and the blob is small we inline it, use EnableForceSubmitRefBlobToBlobStore(false) to disable the backup submission into the blob store
+			byte[] blobPayloadBytes = payload.GetView().ToArray();
+			bool putIntoBlobStore = _cloudDDCSettings.CurrentValue.EnableForceSubmitRefBlobToBlobStore;
+			bool inlineBlob = _cloudDDCSettings.CurrentValue.EnableInlineSmallBlobs;
+			if (inlineBlob && blobPayloadBytes.LongLength > _cloudDDCSettings.CurrentValue.InlineBlobMaxSize)
+			{
+				// do not inline large blobs, instead put them into the blob store
+				blobPayloadBytes = Array.Empty<byte>();
+				putIntoBlobStore = true;
+			}
+			else if (!inlineBlob)
+			{
+				putIntoBlobStore = true;
+			}
+
+			Task objectStorePut = _referencesStore.PutAsync(ns, bucket, key, blobHash, blobPayloadBytes, isFinalized, cancellationToken);
 
 			Task<BlobId>? blobStorePut = null;
-			if (_cloudDDCSettings.CurrentValue.EnablePutRefBodyIntoBlobStore)
+			if (putIntoBlobStore)
 			{
 				blobStorePut = _blobService.PutObjectAsync(ns, payload.GetView().ToArray(), blobHash, cancellationToken);
 			}
diff --git a/Engine/Source/Programs/UnrealCloudDDC/Jupiter/Settings.cs b/Engine/Source/Programs/UnrealCloudDDC/Jupiter/Settings.cs
--- a/Engine/Source/Programs/UnrealCloudDDC/Jupiter/Settings.cs
+++ b/Engine/Source/Programs/UnrealCloudDDC/Jupiter/Settings.cs
@@ -115,8 +115,15 @@ namespace Jupiter
 		public bool EnableOnDemandReplication { get; set; } = true;
 
 		public bool EnableBucketStatsTracking { get; set; } = true;
-		public bool EnablePutRefBodyIntoBlobStore { get; set; } = true;
+		public bool EnableInlineSmallBlobs { get; set; } = true;
+
+		/// <summary>
+		/// Forces the inlined blobs to also be submitted into the blob store, is the old behavior and is not recommended.
+		/// </summary>
+		public bool EnableForceSubmitRefBlobToBlobStore { get; set; } = true;
 		public bool RequirePrivatePortForEnumeration { get; set; } = true;
+
+		public long InlineBlobMaxSize { get; set; } = 32 * 1024; // default to 32 kb blobs max
 	}
 
 	public class MongoSettings
diff --git a/Engine/Source/Programs/UnrealCloudDDC/changelog.md b/Engine/Source/Programs/UnrealCloudDDC/changelog.md
--- a/Engine/Source/Programs/UnrealCloudDDC/changelog.md
+++ b/Engine/Source/Programs/UnrealCloudDDC/changelog.md
@@ -1,4 +1,5 @@
 # Unreleased
+* `InlineMaxBlobSize` option has moved from `Scylla` to `UnrealCloudDDC` as this can now be used to generally control if blob are inlined into the ref store (minimal practical difference as only Scylla supports inlining blobs). Resolves issue with `EnablePutRefBodyIntoBlobStore` option which will prevent these inlined blobs from also being added to the blob store.
 
 # 1.1.1
 * Fixes to helm chart when using ServiceAccounts for authentication and configuring replication
