Unreal Cloud DDC - Fixes issue when trying to get a compact binary of the bucket enumeration

diff --git a/Engine/Source/Programs/UnrealCloudDDC/Jupiter/Controllers/ReferencesController.cs b/Engine/Source/Programs/UnrealCloudDDC/Jupiter/Controllers/ReferencesController.cs
--- a/Engine/Source/Programs/UnrealCloudDDC/Jupiter/Controllers/ReferencesController.cs
+++ b/Engine/Source/Programs/UnrealCloudDDC/Jupiter/Controllers/ReferencesController.cs
@@ -1343,6 +1343,11 @@ namespace Jupiter.Controllers
 
 	public class EnumerateBucketResponse
 	{
+		public EnumerateBucketResponse()
+		{
+			RefIds = null!;
+		}
+
 		public EnumerateBucketResponse(List<RefId> refIds)
 		{
 			RefIds = refIds;
