Unreal Cloud DDC - Decreased time between bucket stat metric generation, was every 24 hrs which means we could not even have data once per day if we restart the worker (or deploy a new version). Is now 6 hours (still configurable) by default which still avoids putting alot of pressure on the db.

diff --git a/Engine/Source/Programs/UnrealCloudDDC/Jupiter/Implementation/MetricsService.cs b/Engine/Source/Programs/UnrealCloudDDC/Jupiter/Implementation/MetricsService.cs
--- a/Engine/Source/Programs/UnrealCloudDDC/Jupiter/Implementation/MetricsService.cs
+++ b/Engine/Source/Programs/UnrealCloudDDC/Jupiter/Implementation/MetricsService.cs
@@ -23,11 +23,11 @@ namespace Jupiter.Implementation
 	public class MetricsServiceSettings
 	{
 		/// <summary>
-		/// Set to enable calulcation of metrics in the background. Adds load to database so only enable these if you intend to use it.
+		/// Set to enable calculation of metrics in the background. Adds load to database so only enable these if you intend to use it.
 		/// </summary>
 		public bool Enabled { get; set; } = false;
 
-		public TimeSpan PollFrequency { get; set; } = TimeSpan.FromHours(24);
+		public TimeSpan PollFrequency { get; set; } = TimeSpan.FromHours(6);
 	}
 
 	public class MetricsService : PollingService<MetricsState>
