Unreal Cloud DDC - Added ability to override max keys returned from s3 when listing blobs - will experiment to see if this speeds up blob gc

diff --git a/Engine/Source/Programs/UnrealCloudDDC/Jupiter/Implementation/Blob/AmazonS3Store.cs b/Engine/Source/Programs/UnrealCloudDDC/Jupiter/Implementation/Blob/AmazonS3Store.cs
--- a/Engine/Source/Programs/UnrealCloudDDC/Jupiter/Implementation/Blob/AmazonS3Store.cs
+++ b/Engine/Source/Programs/UnrealCloudDDC/Jupiter/Implementation/Blob/AmazonS3Store.cs
@@ -432,7 +432,7 @@ namespace Jupiter.Implementation
 					{
 						BucketName = _bucketName,
 						Prefix = hashPrefix,
-						MaxKeys = 10_000
+						MaxKeys = _settings.CurrentValue.PerPrefixMaxKeys
 					};
 
 					ListObjectsV2Response response;
diff --git a/Engine/Source/Programs/UnrealCloudDDC/Jupiter/Settings.cs b/Engine/Source/Programs/UnrealCloudDDC/Jupiter/Settings.cs
--- a/Engine/Source/Programs/UnrealCloudDDC/Jupiter/Settings.cs
+++ b/Engine/Source/Programs/UnrealCloudDDC/Jupiter/Settings.cs
@@ -203,6 +203,11 @@ namespace Jupiter
 		/// Is also a speed up on larger datasets
 		/// </summary>
 		public bool PerPrefixListing { get; set; } = true;
+
+		/// <summary>
+		/// Max number of keys returned in a single request when listing S3
+		/// </summary>
+		public int PerPrefixMaxKeys { get; set; } = 10_000;
 	}
 
 	public class GCSettings
