Unreal Cloud DDC - Added overload to delete blob from multiple namespaces at once - this reduces count of operations against s3 when deleting blobs in the same storage pool and speeds up the deletion

Removed ability to enumerate policies as this didn't correctly take into account the default values of policies and had bad assumptions - is not used anyway.

diff --git a/Engine/Source/Programs/UnrealCloudDDC/Jupiter.Common/Implementation/NamespacePolicyResolver.cs b/Engine/Source/Programs/UnrealCloudDDC/Jupiter.Common/Implementation/NamespacePolicyResolver.cs
--- a/Engine/Source/Programs/UnrealCloudDDC/Jupiter.Common/Implementation/NamespacePolicyResolver.cs
+++ b/Engine/Source/Programs/UnrealCloudDDC/Jupiter.Common/Implementation/NamespacePolicyResolver.cs
@@ -9,7 +9,6 @@ namespace Jupiter.Common
 {
 	public interface INamespacePolicyResolver
 	{
-		public IEnumerable<(NamespaceId, NamespacePolicy)> GetAllPolicies();
 		public NamespacePolicy GetPoliciesForNs(NamespaceId ns);
 
 		static NamespaceId JupiterInternalNamespace => new NamespaceId("jupiter-internal");
@@ -84,14 +83,6 @@ namespace Jupiter.Common
 			}
 		}
 
-		public IEnumerable<(NamespaceId, NamespacePolicy)> GetAllPolicies()
-		{
-			foreach (KeyValuePair<string, NamespacePolicy> pair in _namespaceSettings.CurrentValue.Policies)
-			{
-				yield return (new NamespaceId(pair.Key), pair.Value);
-			}
-		}
-
 		public NamespacePolicy GetPoliciesForNs(NamespaceId ns)
 		{
 			if (_namespaceSettings.CurrentValue.Policies.TryGetValue(ns.ToString(), out NamespacePolicy? settings))
diff --git a/Engine/Source/Programs/UnrealCloudDDC/Jupiter/Implementation/Blob/AmazonS3Store.cs b/Engine/Source/Programs/UnrealCloudDDC/Jupiter/Implementation/Blob/AmazonS3Store.cs
--- a/Engine/Source/Programs/UnrealCloudDDC/Jupiter/Implementation/Blob/AmazonS3Store.cs
+++ b/Engine/Source/Programs/UnrealCloudDDC/Jupiter/Implementation/Blob/AmazonS3Store.cs
@@ -4,6 +4,7 @@ using System;
 using System.Collections.Concurrent;
 using System.Collections.Generic;
 using System.IO;
+using System.Linq;
 using System.Net;
 using System.Runtime.CompilerServices;
 using System.Threading;
@@ -196,6 +197,19 @@ namespace Jupiter.Implementation
 			IStorageBackend backend = GetBackend(ns);
 			await backend.DeleteAsync(blobIdentifier.AsS3Key());
 		}
+
+		public async Task DeleteObjectAsync(IEnumerable<NamespaceId> namespaces, BlobId blob)
+		{
+			List<NamespaceId> namespaceIds = namespaces.ToList();
+			IEnumerable<string> storagePools = namespaceIds.Select(ns => _namespacePolicyResolver.GetPoliciesForNs(ns).StoragePool);
+
+			Dictionary<string, NamespaceId> storagePoolsToClean = storagePools.ToDictionary(storagePool => storagePool, storagePool => namespaceIds.FirstOrDefault(id => _namespacePolicyResolver.GetPoliciesForNs(id).StoragePool == storagePool));
+
+			foreach ((string _, NamespaceId ns) in storagePoolsToClean)
+			{
+				await GetBackend(ns).DeleteAsync(blob.AsS3Key(), CancellationToken.None);
+			}
+		}
 	}
 
 	public class AmazonStorageBackend : IStorageBackend
diff --git a/Engine/Source/Programs/UnrealCloudDDC/Jupiter/Implementation/Blob/AzureBlobStore.cs b/Engine/Source/Programs/UnrealCloudDDC/Jupiter/Implementation/Blob/AzureBlobStore.cs
--- a/Engine/Source/Programs/UnrealCloudDDC/Jupiter/Implementation/Blob/AzureBlobStore.cs
+++ b/Engine/Source/Programs/UnrealCloudDDC/Jupiter/Implementation/Blob/AzureBlobStore.cs
@@ -5,6 +5,7 @@ using System.Collections.Concurrent;
 using System.Collections.Generic;
 using System.Globalization;
 using System.IO;
+using System.Linq;
 using System.Runtime.CompilerServices;
 using System.Threading;
 using System.Threading.Tasks;
@@ -146,6 +147,19 @@ namespace Jupiter.Implementation
 			await GetBackend(ns).DeleteAsync(GetPath(blobIdentifier), CancellationToken.None);
 		}
 
+		public async Task DeleteObjectAsync(IEnumerable<NamespaceId> namespaces, BlobId blob)
+		{
+			List<NamespaceId> namespaceIds = namespaces.ToList();
+			List<string> storagePools = namespaceIds.Select(ns => _namespacePolicyResolver.GetPoliciesForNs(ns).StoragePool).Distinct().ToList();
+
+			Dictionary<string, NamespaceId> storagePoolsToClean = storagePools.ToDictionary(storagePool => storagePool, storagePool => namespaceIds.FirstOrDefault(id => _namespacePolicyResolver.GetPoliciesForNs(id).StoragePool == storagePool));
+
+			foreach ((string _, NamespaceId ns) in storagePoolsToClean)
+			{
+				await GetBackend(ns).DeleteAsync(GetPath(blob), CancellationToken.None);
+			}
+		}
+
 		public async Task DeleteNamespaceAsync(NamespaceId ns)
 		{
 			string fixedNamespace = GetContainerName(ns);
diff --git a/Engine/Source/Programs/UnrealCloudDDC/Jupiter/Implementation/Blob/BlobService.cs b/Engine/Source/Programs/UnrealCloudDDC/Jupiter/Implementation/Blob/BlobService.cs
--- a/Engine/Source/Programs/UnrealCloudDDC/Jupiter/Implementation/Blob/BlobService.cs
+++ b/Engine/Source/Programs/UnrealCloudDDC/Jupiter/Implementation/Blob/BlobService.cs
@@ -825,6 +825,31 @@ public class BlobService : IBlobService
 		throw new NamespaceNotFoundException(ns);
 	}
 
+	public async Task DeleteObjectAsync(List<NamespaceId> namespaces, BlobId blob, CancellationToken cancellationToken = default)
+	{
+		// remove the object from the tracking first, if this times out we do not want to end up with a inconsistent blob index
+		// if the blob store delete fails on the other hand we will still run a delete again during GC (as the blob is still orphaned at that point)
+		// this assumes that blob gc is based on scanning the root blob store
+
+		await Parallel.ForEachAsync(namespaces, cancellationToken, async (ns, token) =>
+		{
+			await _blobIndex.RemoveBlobFromRegionAsync(ns, blob, cancellationToken: cancellationToken);
+			await _blobIndex.RemoveReferencesAsync(ns, blob, null, cancellationToken);
+		});
+
+		// let each blob store figure out how it could effectively clear a storage pool
+		foreach (IBlobStore store in _blobStores)
+		{
+			using TelemetrySpan scope = _tracer.StartActiveSpan("HierarchicalStore.DeleteObjectFromStoragePool")
+					.SetAttribute("operation.name", "HierarchicalStore.DeleteObjectFromStoragePool")
+					.SetAttribute("resource.name", blob.ToString())
+					.SetAttribute("BlobStore", store.GetType().Name)
+				;
+
+			await store.DeleteObjectAsync(namespaces, blob);
+		}
+	}
+
 	public async Task DeleteNamespaceAsync(NamespaceId ns, CancellationToken cancellationToken = default)
 	{
 		bool deletedAtLeastOnce = false;
diff --git a/Engine/Source/Programs/UnrealCloudDDC/Jupiter/Implementation/Blob/FileSystemStore.cs b/Engine/Source/Programs/UnrealCloudDDC/Jupiter/Implementation/Blob/FileSystemStore.cs
--- a/Engine/Source/Programs/UnrealCloudDDC/Jupiter/Implementation/Blob/FileSystemStore.cs
+++ b/Engine/Source/Programs/UnrealCloudDDC/Jupiter/Implementation/Blob/FileSystemStore.cs
@@ -131,6 +131,15 @@ namespace Jupiter.Implementation
 			await GetBackend(ns).DeleteAsync(path, CancellationToken.None);
 		}
 
+		public async Task DeleteObjectAsync(IEnumerable<NamespaceId> namespaces, BlobId objectName)
+		{
+			string path = GetFilesystemPath(objectName);
+			foreach (NamespaceId ns in namespaces)
+			{
+				await GetBackend(ns).DeleteAsync(path, CancellationToken.None);
+			}
+		}
+		
 		public Task DeleteNamespaceAsync(NamespaceId ns)
 		{
 			DirectoryInfo namespaceDirectory = GetFilesystemPath(ns).ToDirectoryInfo();
diff --git a/Engine/Source/Programs/UnrealCloudDDC/Jupiter/Implementation/Blob/IBlobService.cs b/Engine/Source/Programs/UnrealCloudDDC/Jupiter/Implementation/Blob/IBlobService.cs
--- a/Engine/Source/Programs/UnrealCloudDDC/Jupiter/Implementation/Blob/IBlobService.cs
+++ b/Engine/Source/Programs/UnrealCloudDDC/Jupiter/Implementation/Blob/IBlobService.cs
@@ -44,6 +44,14 @@ public interface IBlobService
 	// Delete a object
 	Task DeleteObjectAsync(NamespaceId ns, BlobId blob, CancellationToken cancellationToken = default);
 
+	/// <summary>
+	/// Delete the object from multiple namespaces at once, extra efficient if they are all in the same storage pool
+	/// </summary>
+	/// <param name="namespaces">The namespaces to delete the blob from</param>
+	/// <param name="blob">The identifier of the blob to delete</param>
+	/// <param name="cancellationToken"></param>
+	Task DeleteObjectAsync(List<NamespaceId> namespaces, BlobId blob, CancellationToken cancellationToken = default);
+
 	// delete the whole namespace
 	Task DeleteNamespaceAsync(NamespaceId ns, CancellationToken cancellationToken = default);
 
diff --git a/Engine/Source/Programs/UnrealCloudDDC/Jupiter/Implementation/Blob/IBlobStore.cs b/Engine/Source/Programs/UnrealCloudDDC/Jupiter/Implementation/Blob/IBlobStore.cs
--- a/Engine/Source/Programs/UnrealCloudDDC/Jupiter/Implementation/Blob/IBlobStore.cs
+++ b/Engine/Source/Programs/UnrealCloudDDC/Jupiter/Implementation/Blob/IBlobStore.cs
@@ -27,6 +27,9 @@ namespace Jupiter.Implementation
 		// Delete a object
 		Task DeleteObjectAsync(NamespaceId ns, BlobId blob);
 
+		// Delete a object from multiple namespaces at once
+		Task DeleteObjectAsync(IEnumerable<NamespaceId> ns, BlobId blob);
+
 		// delete the whole namespace
 		Task DeleteNamespaceAsync(NamespaceId ns);
 
diff --git a/Engine/Source/Programs/UnrealCloudDDC/Jupiter/Implementation/Blob/MemoryBlobStore.cs b/Engine/Source/Programs/UnrealCloudDDC/Jupiter/Implementation/Blob/MemoryBlobStore.cs
--- a/Engine/Source/Programs/UnrealCloudDDC/Jupiter/Implementation/Blob/MemoryBlobStore.cs
+++ b/Engine/Source/Programs/UnrealCloudDDC/Jupiter/Implementation/Blob/MemoryBlobStore.cs
@@ -89,6 +89,16 @@ namespace Jupiter.Implementation
 
 		public Task DeleteObjectAsync(NamespaceId ns, BlobId blob) => GetBackend(ns).DeleteAsync(GetPath(blob));
 
+		public Task DeleteObjectAsync(IEnumerable<NamespaceId> namespaces, BlobId blob)
+		{
+			foreach (NamespaceId ns in namespaces)
+			{
+				DeleteObjectAsync(ns, blob);
+			}
+
+			return Task.CompletedTask;
+		}
+
 		public Task<bool> ExistsAsync(NamespaceId ns, BlobId blob, bool forceCheck = false) => GetBackend(ns).ExistsAsync(GetPath(blob));
 
 		public Task DeleteNamespaceAsync(NamespaceId ns)
diff --git a/Engine/Source/Programs/UnrealCloudDDC/Jupiter/Implementation/Blob/PeerBlobStore.cs b/Engine/Source/Programs/UnrealCloudDDC/Jupiter/Implementation/Blob/PeerBlobStore.cs
--- a/Engine/Source/Programs/UnrealCloudDDC/Jupiter/Implementation/Blob/PeerBlobStore.cs
+++ b/Engine/Source/Programs/UnrealCloudDDC/Jupiter/Implementation/Blob/PeerBlobStore.cs
@@ -207,6 +207,12 @@ namespace Jupiter.Implementation
 			return Task.CompletedTask;
 		}
 
+		public Task DeleteObjectAsync(IEnumerable<NamespaceId> namespaces, BlobId blob)
+		{
+			// not applicable
+			return Task.CompletedTask;
+		}
+
 		public Task DeleteNamespaceAsync(NamespaceId ns)
 		{
 			// not applicable
diff --git a/Engine/Source/Programs/UnrealCloudDDC/Jupiter/Implementation/Blob/RelayBlobStore.cs b/Engine/Source/Programs/UnrealCloudDDC/Jupiter/Implementation/Blob/RelayBlobStore.cs
--- a/Engine/Source/Programs/UnrealCloudDDC/Jupiter/Implementation/Blob/RelayBlobStore.cs
+++ b/Engine/Source/Programs/UnrealCloudDDC/Jupiter/Implementation/Blob/RelayBlobStore.cs
@@ -106,6 +106,11 @@ namespace Jupiter.Implementation
 			return true;
 		}
 
+		public Task DeleteObjectAsync(IEnumerable<NamespaceId> namespaces, BlobId blob)
+		{
+			throw new NotImplementedException("DeleteObjectAsync from multiple namespaces is not supported on the relay blob store");
+		}
+
 		public Task DeleteObjectAsync(NamespaceId ns, BlobId blob)
 		{
 			throw new NotImplementedException("DeleteObjects is not supported on the relay blob store");
diff --git a/Engine/Source/Programs/UnrealCloudDDC/Jupiter/Implementation/GC/OrphanBlobCleanupRefs.cs b/Engine/Source/Programs/UnrealCloudDDC/Jupiter/Implementation/GC/OrphanBlobCleanupRefs.cs
--- a/Engine/Source/Programs/UnrealCloudDDC/Jupiter/Implementation/GC/OrphanBlobCleanupRefs.cs
+++ b/Engine/Source/Programs/UnrealCloudDDC/Jupiter/Implementation/GC/OrphanBlobCleanupRefs.cs
@@ -219,20 +219,9 @@ namespace Jupiter.Implementation
 			_deletedBlobCounter.Add(1);
 			_logger.LogInformation("GC Orphan blob {Blob} from {StoragePool} which was last modified at {LastModifiedTime}", blob, storagePoolName, lastModifiedTime);
 
-			// if the blob was not found to have a reference in any of the namespace that share a storage pool then the blob is not used anymore and should be deleted from all the namespaces
-			await Parallel.ForEachAsync(namespacesThatSharePool, cancellationToken, async (ns, ctx) =>
-			{
-				await RemoveBlobAsync(ns, blob, ctx);
-			});
-			return true;
-
-		}
-
-		private async Task RemoveBlobAsync(NamespaceId ns, BlobId blob, CancellationToken cancellationToken)
-		{
 			try
 			{
-				await _blobService.DeleteObjectAsync(ns, blob, cancellationToken);
+				await _blobService.DeleteObjectAsync(namespacesThatSharePool, blob, cancellationToken);
 			}
 			catch (BlobNotFoundException)
 			{
@@ -240,8 +229,11 @@ namespace Jupiter.Implementation
 			}
 			catch (Exception e)
 			{
-				_logger.LogWarning("Failed to delete blob {Blob} from {Namespace} due to {Error}", blob, ns, e.Message);
+				_logger.LogWarning("Failed to delete blob {Blob} from {StoragePool} due to {Error}", blob, storagePoolName, e.Message);
 			}
+			
+			return true;
+
 		}
 
 		private IAsyncEnumerable<NamespaceId> ListNamespaces()
