Unreal Cloud DDC - Added admin endpoint to fetch information on the replicated state for a ref
Defaulted the cluster settings to a single "local" region when no settings have been specified - makes more sense then acting like nothing exists.

diff --git a/Engine/Source/Programs/UnrealCloudDDC/Jupiter/Controllers/ReferencesController.cs b/Engine/Source/Programs/UnrealCloudDDC/Jupiter/Controllers/ReferencesController.cs
--- a/Engine/Source/Programs/UnrealCloudDDC/Jupiter/Controllers/ReferencesController.cs
+++ b/Engine/Source/Programs/UnrealCloudDDC/Jupiter/Controllers/ReferencesController.cs
@@ -21,6 +21,7 @@ using JetBrains.Annotations;
 using Jupiter.Common;
 using Jupiter.Common.Implementation;
 using Jupiter.Implementation;
+using Jupiter.Implementation.Blob;
 using Jupiter.Utils;
 using Microsoft.AspNetCore.Authorization;
 using Microsoft.AspNetCore.Http;
@@ -58,11 +59,15 @@ namespace Jupiter.Controllers
 		private readonly IOptionsMonitor<UnrealCloudDDCSettings> _settings;
 		private readonly IRefService _refService;
 		private readonly IBlobService _blobStore;
+		private readonly IBlobIndex _blobIndex;
+		private readonly IPeerStatusService _peerStatusService;
 
-		public ReferencesController(IRefService refService, IBlobService blobStore, IDiagnosticContext diagnosticContext, FormatResolver formatResolver, BufferedPayloadFactory bufferedPayloadFactory, IReferenceResolver referenceResolver, NginxRedirectHelper nginxRedirectHelper, IRequestHelper requestHelper, Tracer tracer, ILogger<ReferencesController> logger, IOptionsMonitor<UnrealCloudDDCSettings> settings)
+		public ReferencesController(IRefService refService, IBlobService blobStore, IBlobIndex blobIndex, IPeerStatusService peerStatusService, IDiagnosticContext diagnosticContext, FormatResolver formatResolver, BufferedPayloadFactory bufferedPayloadFactory, IReferenceResolver referenceResolver, NginxRedirectHelper nginxRedirectHelper, IRequestHelper requestHelper, Tracer tracer, ILogger<ReferencesController> logger, IOptionsMonitor<UnrealCloudDDCSettings> settings)
 		{
 			_refService = refService;
 			_blobStore = blobStore;
+			_blobIndex = blobIndex;
+			_peerStatusService = peerStatusService;
 			_diagnosticContext = diagnosticContext;
 			_formatResolver = formatResolver;
 			_bufferedPayloadFactory = bufferedPayloadFactory;
@@ -501,6 +506,58 @@ namespace Jupiter.Controllers
 			}
 		}
 
+		/// <summary>
+		/// Returns the replicated state of this ref across all known cloud ddc regions
+		/// </summary>
+		/// <param name="ns">Namespace. Each namespace is completely separated from each other. Use for different types of data that is never expected to be similar (between two different games for instance). Example: `uc4.ddc`</param>
+		/// <param name="bucket">The category/type of record you are caching. Is a clustered key together with the actual key, but all records in the same bucket can be dropped easily. Example: `terrainTexture` </param>
+		/// <param name="key">The unique name of this particular key. `iAmAVeryValidKey`</param>
+		[HttpGet("{ns}/{bucket}/{key}/replicationState", Order = 500)]
+		public async Task<IActionResult> GetReplicationStateAsync(
+			[FromRoute][Required] NamespaceId ns,
+			[FromRoute][Required] BucketId bucket,
+			[FromRoute][Required] RefId key)
+		{
+			ActionResult? accessResult = await _requestHelper.HasAccessToNamespaceAsync(User, Request, ns, new[] { JupiterAclAction.AdminAction });
+			if (accessResult != null)
+			{
+				return accessResult;
+			}
+
+			try
+			{
+				List<BlobId> blobs = await _refService.GetReferencedBlobsAsync(ns, bucket, key, ignoreMissingBlobs: true);
+				ConcurrentDictionary<string, ConcurrentDictionary<string, bool>> blobStatePerRegion = new ConcurrentDictionary<string, ConcurrentDictionary<string, bool>>();
+
+				await Parallel.ForEachAsync(blobs, async (blobId, cancellationToken) =>
+				{
+					ConcurrentDictionary<string, bool> blobState = new ConcurrentDictionary<string, bool>();
+
+					foreach (string region in _peerStatusService.GetRegions())
+					{
+						bool exists = await _blobIndex.BlobExistsInRegionAsync(ns, blobId, region, CancellationToken.None);
+						blobState.TryAdd(region, exists);
+					}
+
+					blobStatePerRegion.TryAdd(blobId.ToString(), blobState);
+				});
+
+				return Ok(new RefReplicationStateResponse(ns, bucket, key, blobStatePerRegion));
+			}
+			catch (NamespaceNotFoundException e)
+			{
+				return NotFound(new ProblemDetails { Title = $"Namespace {e.Namespace} did not exist" });
+			}
+			catch (RefNotFoundException e)
+			{
+				return NotFound(new ProblemDetails { Title = $"Object {e.Bucket} {e.Key} did not exist" });
+			}
+			catch (BlobNotFoundException e)
+			{
+				return NotFound(new ProblemDetails { Title = $"Object {e.Blob} in {e.Ns} not found" });
+			}
+		}
+
 		/// <summary>
 		/// Checks if a object exists
 		/// </summary>
@@ -1165,6 +1222,22 @@ namespace Jupiter.Controllers
 		}
 	}
 
+	public class RefReplicationStateResponse
+	{
+		public RefReplicationStateResponse(NamespaceId ns, BucketId bucket, RefId key, ConcurrentDictionary<string, ConcurrentDictionary<string, bool>> blobsPerRegion)
+		{
+			Ns = ns;
+			Bucket = bucket;
+			Key = key;
+			BlobsPerRegion = blobsPerRegion;
+		}
+
+		public NamespaceId Ns { get;set; }
+		public BucketId Bucket { get;set; }
+		public RefId Key { get;set; }
+		public ConcurrentDictionary<string, ConcurrentDictionary<string, bool>> BlobsPerRegion { get; init; }
+	}
+
 	public class RefDeletedResponse
 	{
 		public RefDeletedResponse()
diff --git a/Engine/Source/Programs/UnrealCloudDDC/Jupiter/Implementation/PeerStatusService.cs b/Engine/Source/Programs/UnrealCloudDDC/Jupiter/Implementation/PeerStatusService.cs
--- a/Engine/Source/Programs/UnrealCloudDDC/Jupiter/Implementation/PeerStatusService.cs
+++ b/Engine/Source/Programs/UnrealCloudDDC/Jupiter/Implementation/PeerStatusService.cs
@@ -3,12 +3,14 @@
 using System;
 using System.Collections.Generic;
 using System.Diagnostics;
+using System.Linq;
 using System.Net.Http;
 using System.Threading;
 using System.Threading.Tasks;
 using EpicGames.Core;
 using Microsoft.Extensions.Logging;
 using Microsoft.Extensions.Options;
+using MongoDB.Driver.Core.Configuration;
 
 namespace Jupiter.Implementation
 {
@@ -26,6 +28,17 @@ namespace Jupiter.Implementation
 
 	public interface IPeerStatusService
 	{
+		/// <summary>
+		/// Lists all known regions
+		/// </summary>
+		/// <returns></returns>
+		IEnumerable<string> GetRegions();
+
+		/// <summary>
+		/// Returns all the information about a particular region
+		/// </summary>
+		/// <param name="regionName"></param>
+		/// <returns></returns>
 		PeerStatus? GetPeerStatus(string regionName);
 
 		/// <summary>
@@ -39,6 +52,7 @@ namespace Jupiter.Implementation
 	public class PeerStatusServiceState
 	{
 	}
+
 	public class PeerStatusService : PollingService<PeerStatusServiceState>, IPeerStatusService
 	{
 		private readonly IOptionsMonitor<ClusterSettings> _clusterSettings;
@@ -48,6 +62,11 @@ namespace Jupiter.Implementation
 		private volatile bool _alreadyPolling = false;
 		private readonly ILogger _logger;
 
+		public IEnumerable<string> GetRegions()
+		{
+			return _clusterSettings.CurrentValue.Peers.Select(settings => settings.Name);
+		}
+
 		public PeerStatus? GetPeerStatus(string regionName)
 		{
 			if (_peers.TryGetValue(regionName, out PeerStatus? peerStatus))
diff --git a/Engine/Source/Programs/UnrealCloudDDC/Jupiter/Settings.cs b/Engine/Source/Programs/UnrealCloudDDC/Jupiter/Settings.cs
--- a/Engine/Source/Programs/UnrealCloudDDC/Jupiter/Settings.cs
+++ b/Engine/Source/Programs/UnrealCloudDDC/Jupiter/Settings.cs
@@ -236,7 +236,7 @@ namespace Jupiter
 	public class ClusterSettings
 	{
 		[System.Diagnostics.CodeAnalysis.SuppressMessage("Usage", "CA2227:Collection properties should be read only", Justification = "Used by serialization")]
-		public List<PeerSettings> Peers { get; set; } = new List<PeerSettings>();
+		public List<PeerSettings> Peers { get; set; } = new List<PeerSettings>() { new PeerSettings() { Name = "local"}};
 	}
 
 	public class PeerSettings
