Unreal Cloud DDC - Fixed issue with the default cluster settings, defaulting this in the settings file instead. Also fixed issue with json serialization of concurrent dicts.

diff --git a/Engine/Source/Programs/UnrealCloudDDC/Jupiter/Controllers/ReferencesController.cs b/Engine/Source/Programs/UnrealCloudDDC/Jupiter/Controllers/ReferencesController.cs
--- a/Engine/Source/Programs/UnrealCloudDDC/Jupiter/Controllers/ReferencesController.cs
+++ b/Engine/Source/Programs/UnrealCloudDDC/Jupiter/Controllers/ReferencesController.cs
@@ -527,11 +527,11 @@ namespace Jupiter.Controllers
 			try
 			{
 				List<BlobId> blobs = await _refService.GetReferencedBlobsAsync(ns, bucket, key, ignoreMissingBlobs: true);
-				ConcurrentDictionary<string, ConcurrentDictionary<string, bool>> blobStatePerRegion = new ConcurrentDictionary<string, ConcurrentDictionary<string, bool>>();
+				Dictionary<string, Dictionary<string, bool>> blobStatePerRegion = new Dictionary<string, Dictionary<string, bool>>();
 
 				await Parallel.ForEachAsync(blobs, async (blobId, cancellationToken) =>
 				{
-					ConcurrentDictionary<string, bool> blobState = new ConcurrentDictionary<string, bool>();
+					Dictionary<string, bool> blobState = new Dictionary<string, bool>();
 
 					foreach (string region in _peerStatusService.GetRegions())
 					{
@@ -539,7 +539,10 @@ namespace Jupiter.Controllers
 						blobState.TryAdd(region, exists);
 					}
 
-					blobStatePerRegion.TryAdd(blobId.ToString(), blobState);
+					lock (blobStatePerRegion)
+					{
+						blobStatePerRegion[blobId.ToString()] = blobState;
+					}
 				});
 
 				return Ok(new RefReplicationStateResponse(ns, bucket, key, blobStatePerRegion));
@@ -1224,7 +1227,7 @@ namespace Jupiter.Controllers
 
 	public class RefReplicationStateResponse
 	{
-		public RefReplicationStateResponse(NamespaceId ns, BucketId bucket, RefId key, ConcurrentDictionary<string, ConcurrentDictionary<string, bool>> blobsPerRegion)
+		public RefReplicationStateResponse(NamespaceId ns, BucketId bucket, RefId key, Dictionary<string, Dictionary<string, bool>> blobsPerRegion)
 		{
 			Ns = ns;
 			Bucket = bucket;
@@ -1235,7 +1238,7 @@ namespace Jupiter.Controllers
 		public NamespaceId Ns { get;set; }
 		public BucketId Bucket { get;set; }
 		public RefId Key { get;set; }
-		public ConcurrentDictionary<string, ConcurrentDictionary<string, bool>> BlobsPerRegion { get; init; }
+		public Dictionary<string, Dictionary<string, bool>> BlobsPerRegion { get; init; }
 	}
 
 	public class RefDeletedResponse
diff --git a/Engine/Source/Programs/UnrealCloudDDC/Jupiter/Implementation/PeerStatusService.cs b/Engine/Source/Programs/UnrealCloudDDC/Jupiter/Implementation/PeerStatusService.cs
--- a/Engine/Source/Programs/UnrealCloudDDC/Jupiter/Implementation/PeerStatusService.cs
+++ b/Engine/Source/Programs/UnrealCloudDDC/Jupiter/Implementation/PeerStatusService.cs
@@ -10,7 +10,6 @@ using System.Threading.Tasks;
 using EpicGames.Core;
 using Microsoft.Extensions.Logging;
 using Microsoft.Extensions.Options;
-using MongoDB.Driver.Core.Configuration;
 
 namespace Jupiter.Implementation
 {
diff --git a/Engine/Source/Programs/UnrealCloudDDC/Jupiter/Settings.cs b/Engine/Source/Programs/UnrealCloudDDC/Jupiter/Settings.cs
--- a/Engine/Source/Programs/UnrealCloudDDC/Jupiter/Settings.cs
+++ b/Engine/Source/Programs/UnrealCloudDDC/Jupiter/Settings.cs
@@ -236,7 +236,7 @@ namespace Jupiter
 	public class ClusterSettings
 	{
 		[System.Diagnostics.CodeAnalysis.SuppressMessage("Usage", "CA2227:Collection properties should be read only", Justification = "Used by serialization")]
-		public List<PeerSettings> Peers { get; set; } = new List<PeerSettings>() { new PeerSettings() { Name = "local"}};
+		public List<PeerSettings> Peers { get; set; } = new List<PeerSettings>();
 	}
 
 	public class PeerSettings
diff --git a/Engine/Source/Programs/UnrealCloudDDC/Jupiter/appsettings.Development.json b/Engine/Source/Programs/UnrealCloudDDC/Jupiter/appsettings.Development.json
--- a/Engine/Source/Programs/UnrealCloudDDC/Jupiter/appsettings.Development.json
+++ b/Engine/Source/Programs/UnrealCloudDDC/Jupiter/appsettings.Development.json
@@ -68,6 +68,14 @@
         "InternalApiPorts": [ 57005, 57006 ]
     },
 
+    "Cluster": {
+		"Peers": [
+			{
+				"Name": "local"
+			}
+		]
+    },
+
     // disable auth and allow any namespace
     "Auth": {
         "Enabled": false
