Unreal Cloud DDC - Fixed race condition when cleaning up ref which could cause the stats tracking to fail to update (as ref got deleted before we finished determining which blobs it referenced)

diff --git a/Engine/Source/Programs/UnrealCloudDDC/Jupiter/Implementation/GC/CacheCleanup.cs b/Engine/Source/Programs/UnrealCloudDDC/Jupiter/Implementation/GC/CacheCleanup.cs
--- a/Engine/Source/Programs/UnrealCloudDDC/Jupiter/Implementation/GC/CacheCleanup.cs
+++ b/Engine/Source/Programs/UnrealCloudDDC/Jupiter/Implementation/GC/CacheCleanup.cs
@@ -168,14 +168,17 @@ namespace Jupiter.Implementation
 			bool storeDelete = false;
 			try
 			{
-				Task? bucketStatsCleanupTask = null;
 				if (_cloudDDCSettings.CurrentValue.EnableBucketStatsTracking)
 				{
-					bucketStatsCleanupTask = Task.Run(async () =>
+					try
 					{
 						List<BlobId> blobs = await _objectService.GetReferencedBlobsAsync(ns, bucket, name, ignoreMissingBlobs: true, cancellationToken: cancellationToken);
 						await _blobIndex.RemoveBlobFromBucketListAsync(ns, bucket, name, blobs, cancellationToken);
-					}, cancellationToken);
+					}
+					catch (RefNotFoundException)
+					{
+						// if the ref is already deleted its not possible for us to cleanup the bucket list
+					}
 				}
 
 				storeDelete = await _referencesStore.DeleteAsync(ns, bucket, name, cancellationToken);
@@ -184,11 +187,6 @@ namespace Jupiter.Implementation
 					// insert a delete event into the transaction log
 					await _replicationLog.InsertDeleteEventAsync(ns, bucket, name, null);
 				}
-
-				if (bucketStatsCleanupTask != null)
-				{
-					await bucketStatsCleanupTask;
-				}
 			}
 			catch (Exception e)
 			{
