Unreal Cloud DDC - Fixed issue in replication when namespace had no new data added to it in that region for 2 weeks

diff --git a/Engine/Source/Programs/UnrealCloudDDC/Jupiter/Controllers/ProblemTypes.cs b/Engine/Source/Programs/UnrealCloudDDC/Jupiter/Controllers/ProblemTypes.cs
--- a/Engine/Source/Programs/UnrealCloudDDC/Jupiter/Controllers/ProblemTypes.cs
+++ b/Engine/Source/Programs/UnrealCloudDDC/Jupiter/Controllers/ProblemTypes.cs
@@ -5,5 +5,6 @@ namespace Jupiter.Controllers
 	public static class ProblemTypes
 	{
 		public static string UseSnapshot { get; set; } = "http://jupiter.epicgames.com/replication/useSnapshot";
+		public static string NoDataFound { get; set; } = "http://jupiter.epicgames.com/replication/noReplicationData";
 	}
 }
diff --git a/Engine/Source/Programs/UnrealCloudDDC/Jupiter/Controllers/ReplicationLogController.cs b/Engine/Source/Programs/UnrealCloudDDC/Jupiter/Controllers/ReplicationLogController.cs
--- a/Engine/Source/Programs/UnrealCloudDDC/Jupiter/Controllers/ReplicationLogController.cs
+++ b/Engine/Source/Programs/UnrealCloudDDC/Jupiter/Controllers/ReplicationLogController.cs
@@ -123,6 +123,7 @@ namespace Jupiter.Controllers
 				return BadRequest(new ProblemDetails
 				{
 					Title = $"No snapshot or bucket found for namespace \"{ns}\"",
+					Type = ProblemTypes.NoDataFound,
 				});
 			}
 			catch (NamespaceNotFoundException)
diff --git a/Engine/Source/Programs/UnrealCloudDDC/Jupiter/Implementation/Replication/RefsReplicator.cs b/Engine/Source/Programs/UnrealCloudDDC/Jupiter/Implementation/Replication/RefsReplicator.cs
--- a/Engine/Source/Programs/UnrealCloudDDC/Jupiter/Implementation/Replication/RefsReplicator.cs
+++ b/Engine/Source/Programs/UnrealCloudDDC/Jupiter/Implementation/Replication/RefsReplicator.cs
@@ -694,6 +694,10 @@ namespace Jupiter.Implementation
 						throw new UseSnapshotException(snapshotBlob, blobNamespace!.Value);
 					}
 
+					if (problemDetails.Type == ProblemTypes.NoDataFound)
+					{
+						yield break;
+					}
 					throw new Exception($"Unknown bad request response. Body: {body}");
 				}
 
