Cloud DDC - Fixed issue were some blob stats were being recorded in ints which is to small to capture the potential value, in practise we shouldn't have run into this for any data we submit yet.

diff --git a/Engine/Source/Programs/UnrealCloudDDC/Jupiter/Implementation/Blob/ScyllaBlobIndex.cs b/Engine/Source/Programs/UnrealCloudDDC/Jupiter/Implementation/Blob/ScyllaBlobIndex.cs
--- a/Engine/Source/Programs/UnrealCloudDDC/Jupiter/Implementation/Blob/ScyllaBlobIndex.cs
+++ b/Engine/Source/Programs/UnrealCloudDDC/Jupiter/Implementation/Blob/ScyllaBlobIndex.cs
@@ -374,12 +374,12 @@ public class ScyllaBlobIndex : IBlobIndex
 		using TelemetrySpan scope = _tracer.BuildScyllaSpan("scylla.calc_bucket_stats")
 			.SetAttribute("resource.name", $"{nsAsString}.{bucketAsString}");
 
-		int totalCountOfRefs = 0;
-		int totalCountOfBlobs = 0;
+		long totalCountOfRefs = 0;
+		long totalCountOfBlobs = 0;
 		long totalSizeOfBlobs = 0;
 
-		List<int> smallestBlobPerPrefix = new List<int>();
-		List<int> largestBlobPerPrefix = new List<int>();
+		List<long> smallestBlobPerPrefix = new List<long>();
+		List<long> largestBlobPerPrefix = new List<long>();
 
 		string[] hashPrefixes = new string[65536];
 		int i = 0;
@@ -408,7 +408,7 @@ public class ScyllaBlobIndex : IBlobIndex
 				RowSet? rowSet = await _session.ExecuteAsync(boundStatement);
 				foreach (Row row in rowSet)
 				{
-					int countOfRefs = (int)(long)row["system.count(reference_id)"];
+					long countOfRefs = (long)row["system.count(reference_id)"];
 
 					Interlocked.Add(ref totalCountOfRefs, countOfRefs);
 				}
@@ -439,9 +439,9 @@ public class ScyllaBlobIndex : IBlobIndex
 						continue;
 					}
 
-					int countOfBlobs = (int)(long)row["system.count(blob_id)"];
-					int smallestBlob = (int)(long)row["system.min(size)"];
-					int largestBlob = (int)(long)row["system.max(size)"];
+					long countOfBlobs = (long)row["system.count(blob_id)"];
+					long smallestBlob = (long)row["system.min(size)"];
+					long largestBlob = (long)row["system.max(size)"];
 					long sumSizeOfBlobs = (long)row["system.sum(size)"];
 
 					Interlocked.Add(ref totalCountOfBlobs, countOfBlobs);
@@ -454,8 +454,8 @@ public class ScyllaBlobIndex : IBlobIndex
 
 			await Task.WhenAll(calcRefStats, calcBlobStats);
 		});
-		int smallestBlobFound = smallestBlobPerPrefix.Any() ? smallestBlobPerPrefix.Min() : 0;
-		int largestBlobFound = largestBlobPerPrefix.Any() ? largestBlobPerPrefix.Max() : 0;
+		long smallestBlobFound = smallestBlobPerPrefix.Any() ? smallestBlobPerPrefix.Min() : 0;
+		long largestBlobFound = largestBlobPerPrefix.Any() ? largestBlobPerPrefix.Max() : 0;
 
 		return new BucketStats
 		{
