Cloud DDC - Fixed incorrect type usage, use Guid over TimeUuid as its a built in type.

diff --git a/Engine/Source/Programs/UnrealCloudDDC/Jupiter/Implementation/TransactionLog/IReplicationLog.cs b/Engine/Source/Programs/UnrealCloudDDC/Jupiter/Implementation/TransactionLog/IReplicationLog.cs
--- a/Engine/Source/Programs/UnrealCloudDDC/Jupiter/Implementation/TransactionLog/IReplicationLog.cs
+++ b/Engine/Source/Programs/UnrealCloudDDC/Jupiter/Implementation/TransactionLog/IReplicationLog.cs
@@ -86,7 +86,7 @@ namespace Jupiter.Implementation
 	public class BlobReplicationLogEvent
 	{
 		[JsonConstructor]
-		public BlobReplicationLogEvent(NamespaceId @namespace, BlobId blob, TimeUuid eventId, string timeBucket, DateTime timestamp, OpType op, BucketId? bucketHint)
+		public BlobReplicationLogEvent(NamespaceId @namespace, BlobId blob, Guid eventId, string timeBucket, DateTime timestamp, OpType op, BucketId? bucketHint)
 		{
 			Namespace = @namespace;
 			Blob = blob;
@@ -108,7 +108,7 @@ namespace Jupiter.Implementation
 		public OpType Op { get; }
 		public DateTime Timestamp { get; }
 		public string TimeBucket { get; }
-		public TimeUuid EventId { get; }
+		public Guid EventId { get; }
 		public BucketId? BucketHint { get; }
 	}
 
