Cloud DDC - Fixed issue were internal api filters were not functioning if using domain sockets between nginx and cloud ddc.

diff --git a/Engine/Source/Programs/UnrealCloudDDC/Jupiter.Common/InternalApiFilterAttribute.cs b/Engine/Source/Programs/UnrealCloudDDC/Jupiter.Common/InternalApiFilterAttribute.cs
--- a/Engine/Source/Programs/UnrealCloudDDC/Jupiter.Common/InternalApiFilterAttribute.cs
+++ b/Engine/Source/Programs/UnrealCloudDDC/Jupiter.Common/InternalApiFilterAttribute.cs
@@ -15,8 +15,9 @@ namespace Jupiter.Implementation
 		{
 			IOptionsMonitor<JupiterSettings>? settings = context.HttpContext.RequestServices.GetService<IOptionsMonitor<JupiterSettings>>();
 
+			// if internal port is set to 0 we just assume all traffic is the internal port, used for tests
 			bool isInternalPort = settings!.CurrentValue.InternalApiPorts.Contains(context.HttpContext.Connection.LocalPort) ||
-				/* unit tests do not run on ports, we consider them always on the internal port */ (context.HttpContext.Connection.LocalPort == 0 && context.HttpContext.Connection.LocalIpAddress == null);
+			                      settings.CurrentValue.InternalApiPorts.Contains(0);
 			if (!isInternalPort)
 			{
 				// this endpoint should only be exposed on the internal port, so we return a 404 as this is not on the internal port
diff --git a/Engine/Source/Programs/UnrealCloudDDC/Jupiter/appsettings.Testing.json b/Engine/Source/Programs/UnrealCloudDDC/Jupiter/appsettings.Testing.json
--- a/Engine/Source/Programs/UnrealCloudDDC/Jupiter/appsettings.Testing.json
+++ b/Engine/Source/Programs/UnrealCloudDDC/Jupiter/appsettings.Testing.json
@@ -557,7 +557,8 @@
 
 	"Jupiter": {
 		"CurrentSite": "test",
-		"AssumeLocalConnectionsHasFullAccess": "true"
+		"AssumeLocalConnectionsHasFullAccess": "true",
+		"InternalApiPorts": [ 0 ] // consider all traffic as internal traffic so we can each the replication endpoints
 	},
 
 	"Kestrel": {
