CloudDDC - Use the 8.0.300 version of the bundled dotnet sdk - fixes build issue

diff --git a/Engine/Source/Programs/UnrealCloudDDC/BuildUnrealCloudDDC.xml b/Engine/Source/Programs/UnrealCloudDDC/BuildUnrealCloudDDC.xml
--- a/Engine/Source/Programs/UnrealCloudDDC/BuildUnrealCloudDDC.xml
+++ b/Engine/Source/Programs/UnrealCloudDDC/BuildUnrealCloudDDC.xml
@@ -9,8 +9,8 @@
 	<Property Name="InformationalVersion" Value="$(InformationalVersion)-PF-$(PreflightChange)" If="'$(PreflightChange)' != ''"/>
 
 	<Property Name="DotNetPath" Value=""/>
-	<Property Name="DotNetPath" Value="$(RootDir)/Engine/Binaries/ThirdParty/DotNet/8.0.100/windows/dotnet.exe" If="$(HostPlatform) == 'Win64'"/>
-	<Property Name="DotNetPath" Value="$(RootDir)/Engine/Binaries/ThirdParty/DotNet/8.0.100/linux/dotnet" If="$(HostPlatform) == 'Linux'"/>
+	<Property Name="DotNetPath" Value="$(RootDir)/Engine/Binaries/ThirdParty/DotNet/8.0.300/windows/dotnet.exe" If="$(HostPlatform) == 'Win64'"/>
+	<Property Name="DotNetPath" Value="$(RootDir)/Engine/Binaries/ThirdParty/DotNet/8.0.300/linux/dotnet" If="$(HostPlatform) == 'Linux'"/>
 
 	<Agent Name="Jupiter" Type="Linux">
 		<Property Name="ComposeFile" Value="$(RootDir)/Engine/Source/Programs/UnrealCloudDDC/Composes/docker-compose-tests.yml"/>
