Cloud DDC - Fixed path for dotnet 8.0.300 (platform names are different from previous version)

diff --git a/Engine/Source/Programs/UnrealCloudDDC/BuildUnrealCloudDDC.xml b/Engine/Source/Programs/UnrealCloudDDC/BuildUnrealCloudDDC.xml
--- a/Engine/Source/Programs/UnrealCloudDDC/BuildUnrealCloudDDC.xml
+++ b/Engine/Source/Programs/UnrealCloudDDC/BuildUnrealCloudDDC.xml
@@ -9,8 +9,8 @@
 	<Property Name="InformationalVersion" Value="$(InformationalVersion)-PF-$(PreflightChange)" If="'$(PreflightChange)' != ''"/>
 
 	<Property Name="DotNetPath" Value=""/>
-	<Property Name="DotNetPath" Value="$(RootDir)/Engine/Binaries/ThirdParty/DotNet/8.0.300/windows/dotnet.exe" If="$(HostPlatform) == 'Win64'"/>
-	<Property Name="DotNetPath" Value="$(RootDir)/Engine/Binaries/ThirdParty/DotNet/8.0.300/linux/dotnet" If="$(HostPlatform) == 'Linux'"/>
+	<Property Name="DotNetPath" Value="$(RootDir)/Engine/Binaries/ThirdParty/DotNet/8.0.300/win-x64/dotnet.exe" If="$(HostPlatform) == 'Win64'"/>
+	<Property Name="DotNetPath" Value="$(RootDir)/Engine/Binaries/ThirdParty/DotNet/8.0.300/linux-x64/dotnet" If="$(HostPlatform) == 'Linux'"/>
 
 	<Agent Name="Jupiter" Type="Linux">
 		<Property Name="ComposeFile" Value="$(RootDir)/Engine/Source/Programs/UnrealCloudDDC/Composes/docker-compose-tests.yml"/>
