Cloud DDC - Removed assumption that we always need to replicate the ref payload - instead the reference endpoint has been updated to include it if it is not inlined.

diff --git a/Engine/Source/Programs/UnrealCloudDDC/Jupiter/Controllers/ReferencesController.cs b/Engine/Source/Programs/UnrealCloudDDC/Jupiter/Controllers/ReferencesController.cs
--- a/Engine/Source/Programs/UnrealCloudDDC/Jupiter/Controllers/ReferencesController.cs
+++ b/Engine/Source/Programs/UnrealCloudDDC/Jupiter/Controllers/ReferencesController.cs
@@ -621,8 +621,14 @@ namespace Jupiter.Controllers
 
 			try
 			{
-				BlobId[] references = await _referenceResolver.GetReferencedBlobsAsync(ns, compactBinaryObject).ToArrayAsync();
-				return Ok(new ResolvedReferencesResult(references));
+				List<BlobId> references = await _referenceResolver.GetReferencedBlobsAsync(ns, compactBinaryObject).ToListAsync();
+
+				if (refRecord.InlinePayload == null)
+				{
+					// payload is not inlined, so it needs to be replicated
+					references.Add(refRecord.BlobIdentifier);
+				}
+				return Ok(new ResolvedReferencesResult(references.ToArray()));
 			}
 			catch (PartialReferenceResolveException e)
 			{
diff --git a/Engine/Source/Programs/UnrealCloudDDC/Jupiter/Implementation/Replication/RefsReplicator.cs b/Engine/Source/Programs/UnrealCloudDDC/Jupiter/Implementation/Replication/RefsReplicator.cs
--- a/Engine/Source/Programs/UnrealCloudDDC/Jupiter/Implementation/Replication/RefsReplicator.cs
+++ b/Engine/Source/Programs/UnrealCloudDDC/Jupiter/Implementation/Replication/RefsReplicator.cs
@@ -537,9 +537,8 @@ namespace Jupiter.Implementation
 				throw new Exception($"Unable to resolve references for object {objectToReplicate} in namespace {ns}");
 			}
 
-			BlobId[] potentialBlobs = new BlobId[refs.References.Length + 1];
+			BlobId[] potentialBlobs = new BlobId[refs.References.Length];
 			Array.Copy(refs.References, potentialBlobs, refs.References.Length);
-			potentialBlobs[^1] = objectToReplicate;
 
 			BlobId[] missingBlobs = await _blobService.FilterOutKnownBlobsAsync(ns, potentialBlobs, cancellationToken);
 			Task[] blobReplicationTasks = new Task[missingBlobs.Length];
