Cloud DDC - Fixed incorrect logging string for BlobIndexConsistencyCheck

diff --git a/Engine/Source/Programs/UnrealCloudDDC/Jupiter/Implementation/Blob/BlobIndexConsistencyCheckService.cs b/Engine/Source/Programs/UnrealCloudDDC/Jupiter/Implementation/Blob/BlobIndexConsistencyCheckService.cs
--- a/Engine/Source/Programs/UnrealCloudDDC/Jupiter/Implementation/Blob/BlobIndexConsistencyCheckService.cs
+++ b/Engine/Source/Programs/UnrealCloudDDC/Jupiter/Implementation/Blob/BlobIndexConsistencyCheckService.cs
@@ -34,7 +34,7 @@ namespace Jupiter.Implementation
 			return _settings.CurrentValue.EnableBlobIndexChecks;
 		}
 
-		public BlobIndexConsistencyCheckService(IOptionsMonitor<ConsistencyCheckSettings> settings, IOptionsMonitor<JupiterSettings> jupiterSettings, ILeaderElection leaderElection, IBlobIndex blobIndex, IBlobService blobService, INamespacePolicyResolver namespacePolicyResolver, Tracer tracer, ILogger<BlobIndexConsistencyCheckService> logger) : base(serviceName: nameof(BlobStoreConsistencyCheckService), TimeSpan.FromSeconds(settings.CurrentValue.ConsistencyCheckPollFrequencySeconds), new BlobIndexConsistencyState(), logger)
+		public BlobIndexConsistencyCheckService(IOptionsMonitor<ConsistencyCheckSettings> settings, IOptionsMonitor<JupiterSettings> jupiterSettings, ILeaderElection leaderElection, IBlobIndex blobIndex, IBlobService blobService, INamespacePolicyResolver namespacePolicyResolver, Tracer tracer, ILogger<BlobIndexConsistencyCheckService> logger) : base(serviceName: nameof(BlobIndexConsistencyCheckService), TimeSpan.FromSeconds(settings.CurrentValue.ConsistencyCheckPollFrequencySeconds), new BlobIndexConsistencyState(), logger)
 		{
 			_settings = settings;
 			_jupiterSettings = jupiterSettings;
