Cloud DDC - Corrected recommended command line for benchmarks

diff --git a/Engine/Source/Programs/UnrealCloudDDC/Benchmarks/readme.md b/Engine/Source/Programs/UnrealCloudDDC/Benchmarks/readme.md
--- a/Engine/Source/Programs/UnrealCloudDDC/Benchmarks/readme.md
+++ b/Engine/Source/Programs/UnrealCloudDDC/Benchmarks/readme.md
@@ -10,7 +10,7 @@ Usage:
 ```
 docker build -t jupiter_benchmark .
 docker run --network host jupiter_benchmark \
---seed --host <url-to-service> \
+--seed --seed-remote --host <url-to-service> \
 --header="Authorization: Bearer <access-token>" \
 all
 ```
