Cloud DDC - Added optional audit log for symbol access

diff --git a/Engine/Source/Programs/UnrealCloudDDC/Jupiter/Controllers/SymbolsController.cs b/Engine/Source/Programs/UnrealCloudDDC/Jupiter/Controllers/SymbolsController.cs
--- a/Engine/Source/Programs/UnrealCloudDDC/Jupiter/Controllers/SymbolsController.cs
+++ b/Engine/Source/Programs/UnrealCloudDDC/Jupiter/Controllers/SymbolsController.cs
@@ -5,6 +5,7 @@ using System.Collections.Generic;
 using System.ComponentModel.DataAnnotations;
 using System.IO;
 using System.Linq;
+using System.Net.Http;
 using System.Net.Mime;
 using System.Text;
 using System.Threading.Tasks;
@@ -21,9 +22,11 @@ using Microsoft.AspNetCore.Http.Extensions;
 using Microsoft.AspNetCore.Mvc;
 using Microsoft.Extensions.DependencyInjection;
 using Microsoft.Extensions.Logging;
+using Microsoft.Extensions.Options;
 using OpenTelemetry.Trace;
 using Serilog;
 using BinaryReader = System.IO.BinaryReader;
+using ILogger = Serilog.ILogger;
 
 namespace Jupiter.Controllers
 {
@@ -42,8 +45,9 @@ namespace Jupiter.Controllers
 		private readonly IRequestHelper _requestHelper;
 		private readonly Tracer _tracer;
 		private readonly ILogger<SymbolsController> _logger;
+		private readonly ILogger? _auditLogger;
 
-		public SymbolsController(IRefService refService, IBlobService blobStore, IDiagnosticContext diagnosticContext, BufferedPayloadFactory bufferedPayloadFactory, FormatResolver formatResolver, NginxRedirectHelper nginxRedirectHelper, IRequestHelper requestHelper, Tracer tracer, ILogger<SymbolsController> logger)
+		public SymbolsController(IRefService refService, IBlobService blobStore, IDiagnosticContext diagnosticContext, BufferedPayloadFactory bufferedPayloadFactory, FormatResolver formatResolver, NginxRedirectHelper nginxRedirectHelper, IRequestHelper requestHelper, Tracer tracer, ILogger<SymbolsController> logger, IOptionsMonitor<SymbolsSettings> symbolSettings)
 		{
 			_refService = refService;
 			_blobStore = blobStore;
@@ -54,6 +58,7 @@ namespace Jupiter.Controllers
 			_requestHelper = requestHelper;
 			_tracer = tracer;
 			_logger = logger;
+			_auditLogger = symbolSettings.CurrentValue.EnableAuditLog ? Serilog.Log.ForContext("SymbolsAuditLog", null) : null;
 		}
 
 		/// <summary>
@@ -76,6 +81,7 @@ namespace Jupiter.Controllers
 				return accessResult;
 			}
 
+			LogAuditEntry(HttpMethod.Get, ns, moduleName, identifier, fileName);
 			BucketId bucket = new BucketId(moduleName);
 			BlobContents? refContents;
 			try
@@ -97,7 +103,7 @@ namespace Jupiter.Controllers
 			CbObject cb = new CbObject(blobMemory);
 			IoHash payloadHash = cb["pdbPayload"].AsBinaryAttachment().Hash;
 
-			BlobContents referencedBlobContents = await _blobStore.GetObjectAsync(ns, BlobId.FromIoHash(payloadHash), null, supportsRedirectUri: true, bucketHint: bucket);
+			BlobContents referencedBlobContents = await _blobStore.GetObjectAsync(ns, BlobId.FromIoHash(payloadHash), storageLayers: null, supportsRedirectUri: true, bucketHint: bucket);
 
 			if (referencedBlobContents.RedirectUri != null)
 			{
@@ -199,6 +205,7 @@ namespace Jupiter.Controllers
 
 			(string pdbIdentifier, int pdbAge) = ExtractModuleInformation(moduleName, decompressedContent);
 
+			LogAuditEntry(HttpMethod.Put, ns, moduleName, pdbIdentifier, moduleName);
 			string filename = moduleName;
 			IoHash attachmentIoHash = attachmentHash.AsIoHash();
 			CbWriter writer = new CbWriter();
@@ -225,6 +232,11 @@ namespace Jupiter.Controllers
 			return Ok(new PutSymbolResponse(moduleName, pdbIdentifier, pdbAge, attachmentIoHash));
 		}
 
+		private void LogAuditEntry(HttpMethod method, NamespaceId ns, string moduleName, string identifier, string fileName)
+		{
+			_auditLogger?.Information("{HttpMethod} '{Namespace}'/'{ModuleName}:{Identifier}' {FileName} IP:{IP} User:{Username} UserAgent: \"{Useragent}\"", method,ns, moduleName, identifier, fileName, Request.HttpContext.Connection.RemoteIpAddress, User?.Identity?.Name ?? "Unknown-user", string.Join(' ', Request.Headers.UserAgent.ToArray()));
+		}
+
 		private static (string, int) ExtractModuleInformation(string moduleName, IBufferedPayload decompressedContent)
 		{
 			using Stream s = decompressedContent.GetStream();
@@ -373,4 +385,9 @@ namespace Jupiter.Controllers
 		public int PdbAge { get; set; }
 		public IoHash PdbPayload { get; set; }
 	}
+
+	public class SymbolsSettings
+	{
+		public bool EnableAuditLog { get; set; }
+	}
 }
diff --git a/Engine/Source/Programs/UnrealCloudDDC/Jupiter/Jupiter.Startup.cs b/Engine/Source/Programs/UnrealCloudDDC/Jupiter/Jupiter.Startup.cs
--- a/Engine/Source/Programs/UnrealCloudDDC/Jupiter/Jupiter.Startup.cs
+++ b/Engine/Source/Programs/UnrealCloudDDC/Jupiter/Jupiter.Startup.cs
@@ -94,6 +94,8 @@ namespace Jupiter
 
 			services.AddOptions<NginxSettings>().Bind(Configuration.GetSection("Nginx")).ValidateDataAnnotations();
 
+			services.AddOptions<SymbolsSettings>().Bind(Configuration.GetSection("Symbols")).ValidateDataAnnotations();
+
 			services.AddOptions<ConsistencyCheckSettings>().Bind(Configuration.GetSection("ConsistencyCheck")).ValidateDataAnnotations();
 
 			services.AddOptions<BufferedPayloadOptions>().Bind(Configuration.GetSection("PayloadBuffering")).ValidateDataAnnotations();
