CloudDDC - Added support for specifying settings under "Symbols" in the configmap (thus from the helm values file)

diff --git a/Engine/Source/Programs/UnrealCloudDDC/Helm/UnrealCloudDDC/templates/configmap.yaml b/Engine/Source/Programs/UnrealCloudDDC/Helm/UnrealCloudDDC/templates/configmap.yaml
--- a/Engine/Source/Programs/UnrealCloudDDC/Helm/UnrealCloudDDC/templates/configmap.yaml
+++ b/Engine/Source/Programs/UnrealCloudDDC/Helm/UnrealCloudDDC/templates/configmap.yaml
@@ -158,6 +158,11 @@ CacheRef:
 {{ $config.CacheRef | toYaml | nindent 2 }}
 {{- end }}
 
+{{- if $config.Symbols }}
+Symbols:
+{{ $config.Symbols | toYaml | nindent 2 }}
+{{- end }}
+
 {{- if and (.Values.global.cloudProvider) (eq .Values.global.cloudProvider "Azure") }}
 {{ include "cloud.settings.azure" .Values }}
 {{- else if and (.Values.global.cloudProvider) (eq .Values.global.cloudProvider "AWS") }}
