Cloud DDC - Added more context to the audit log lines so we can more easily pick them up in monitoring.

diff --git a/Engine/Source/Programs/UnrealCloudDDC/Jupiter/Controllers/SymbolsController.cs b/Engine/Source/Programs/UnrealCloudDDC/Jupiter/Controllers/SymbolsController.cs
--- a/Engine/Source/Programs/UnrealCloudDDC/Jupiter/Controllers/SymbolsController.cs
+++ b/Engine/Source/Programs/UnrealCloudDDC/Jupiter/Controllers/SymbolsController.cs
@@ -25,6 +25,7 @@ using Microsoft.Extensions.Logging;
 using Microsoft.Extensions.Options;
 using OpenTelemetry.Trace;
 using Serilog;
+using Serilog.Context;
 using BinaryReader = System.IO.BinaryReader;
 using ILogger = Serilog.ILogger;
 
@@ -234,7 +235,8 @@ namespace Jupiter.Controllers
 
 		private void LogAuditEntry(HttpMethod method, NamespaceId ns, string moduleName, string identifier, string fileName)
 		{
-			_auditLogger?.Information("{HttpMethod} '{Namespace}'/'{ModuleName}:{Identifier}' {FileName} IP:{IP} User:{Username} UserAgent: \"{Useragent}\"", method,ns, moduleName, identifier, fileName, Request.HttpContext.Connection.RemoteIpAddress, User?.Identity?.Name ?? "Unknown-user", string.Join(' ', Request.Headers.UserAgent.ToArray()));
+			using IDisposable context = LogContext.PushProperty("LogType", "audit");
+			_auditLogger?.Information("{HttpMethod} '{Namespace}'/'{ModuleName}:{Identifier}' {FileName} IP:{IP} User:{Username} UserAgent:\"{Useragent}\"", method,ns, moduleName, identifier, fileName, Request.HttpContext.Connection.RemoteIpAddress, User?.Identity?.Name ?? "Unknown-user", string.Join(' ', Request.Headers.UserAgent.ToArray()));
 		}
 
 		private static (string, int) ExtractModuleInformation(string moduleName, IBufferedPayload decompressedContent)
