Cloud DDC - Made sure nginx sets xff (X-Forwarded-For) headers when forwarding traffic so we retain information about source ip.
Added enricher to always include ue-session header into log entries if we have it.
Added option to remap name claim for Okta which is set by default, fixes issue were username is not correctly resolved when using Okta.
Upgraded Okta.AspNetCore to .net 8 compatible version.
Fixed so the symbols audit log has slightly nicer properties being output.

diff --git a/Engine/Source/Programs/UnrealCloudDDC/Composes/appsettings.Local.yaml b/Engine/Source/Programs/UnrealCloudDDC/Composes/appsettings.Local.yaml
--- a/Engine/Source/Programs/UnrealCloudDDC/Composes/appsettings.Local.yaml
+++ b/Engine/Source/Programs/UnrealCloudDDC/Composes/appsettings.Local.yaml
@@ -12,7 +12,7 @@ Jupiter:
   ChmodDomainSockets: true
 Replication:
   CurrentSite: docker-compose
-Ngunx:
+Nginx:
   UseNginxRedirect: true
 UnrealCloudDDC:
   storageImplementations:
diff --git a/Engine/Source/Programs/UnrealCloudDDC/Helm/UnrealCloudDDC/templates/configmap-nginx.yaml b/Engine/Source/Programs/UnrealCloudDDC/Helm/UnrealCloudDDC/templates/configmap-nginx.yaml
--- a/Engine/Source/Programs/UnrealCloudDDC/Helm/UnrealCloudDDC/templates/configmap-nginx.yaml
+++ b/Engine/Source/Programs/UnrealCloudDDC/Helm/UnrealCloudDDC/templates/configmap-nginx.yaml
@@ -48,12 +48,20 @@ data:
           proxy_set_header X-Jupiter-XAccel-Supported true;
           proxy_set_header X-Jupiter-Port Public;
           proxy_set_header Connection "";
+    
+          # Always set the x-forwarded-for header so we can extract the proper ip of the sender for logging purposes
+          proxy_set_header X-Forwarded-For $proxy_add_x_forwarded_for;
+          
           proxy_pass http://jupiter_backend;
           proxy_next_upstream error timeout http_500;
         }
         location /health {
           proxy_set_header X-Jupiter-Port Public;
           proxy_set_header Connection "";
+
+          # Always set the x-forwarded-for header so we can extract the proper ip of the sender for logging purposes
+          proxy_set_header X-Forwarded-For $proxy_add_x_forwarded_for;
+
           proxy_pass http://jupiter_backend;
           proxy_next_upstream error timeout http_500;
          }
@@ -92,12 +100,20 @@ data:
           proxy_set_header X-Jupiter-XAccel-Supported true;
           proxy_set_header X-Jupiter-Port Corp;
           proxy_set_header Connection "";
+
+          # Always set the x-forwarded-for header so we can extract the proper ip of the sender for logging purposes
+          proxy_set_header X-Forwarded-For $proxy_add_x_forwarded_for;
+
           proxy_pass http://jupiter_backend;
           proxy_next_upstream error timeout http_500;  
         }
         location /health {
           proxy_set_header X-Jupiter-Port Corp;
           proxy_set_header Connection "";
+
+          # Always set the x-forwarded-for header so we can extract the proper ip of the sender for logging purposes
+          proxy_set_header X-Forwarded-For $proxy_add_x_forwarded_for;
+
           proxy_pass http://jupiter_backend;
           proxy_next_upstream error timeout http_500;
         }
@@ -137,12 +153,20 @@ data:
           proxy_set_header X-Jupiter-XAccel-Supported true;
           proxy_set_header X-Jupiter-Port Public;
           proxy_set_header Connection "";
+
+          # Always set the x-forwarded-for header so we can extract the proper ip of the sender for logging purposes
+          proxy_set_header X-Forwarded-For $proxy_add_x_forwarded_for;
+
           proxy_pass http://jupiter_backend;
           proxy_next_upstream error timeout http_500;
         }
         location /health {
           proxy_set_header X-Jupiter-Port Public;
           proxy_set_header Connection "";
+
+          # Always set the x-forwarded-for header so we can extract the proper ip of the sender for logging purposes
+          proxy_set_header X-Forwarded-For $proxy_add_x_forwarded_for;
+
           proxy_pass http://jupiter_backend;
           proxy_next_upstream error timeout http_500;
         }
@@ -182,12 +206,20 @@ data:
           proxy_set_header X-Jupiter-XAccel-Supported true;
           proxy_set_header X-Jupiter-Port Corp;
           proxy_set_header Connection "";
+
+          # Always set the x-forwarded-for header so we can extract the proper ip of the sender for logging purposes
+          proxy_set_header X-Forwarded-For $proxy_add_x_forwarded_for;
+
           proxy_pass http://jupiter_backend;
           proxy_next_upstream error timeout http_500;
         }
         location /health {
           proxy_set_header X-Jupiter-Port Corp;
           proxy_set_header Connection "";
+          
+          # Always set the x-forwarded-for header so we can extract the proper ip of the sender for logging purposes
+          proxy_set_header X-Forwarded-For $proxy_add_x_forwarded_for;
+
           proxy_pass http://jupiter_backend;
           proxy_next_upstream error timeout http_500;
         }
diff --git a/Engine/Source/Programs/UnrealCloudDDC/Jupiter.Common/BaseProgram.cs b/Engine/Source/Programs/UnrealCloudDDC/Jupiter.Common/BaseProgram.cs
--- a/Engine/Source/Programs/UnrealCloudDDC/Jupiter.Common/BaseProgram.cs
+++ b/Engine/Source/Programs/UnrealCloudDDC/Jupiter.Common/BaseProgram.cs
@@ -52,6 +52,7 @@ namespace Jupiter
 			Log.Logger = new LoggerConfiguration()
 				.ReadFrom.Configuration(Configuration)
 				.Enrich.With<DatadogLogEnricher>()
+				.Enrich.WithRequestHeader("ue-session", "ue-session")
 				.CreateLogger();
 
 			try
diff --git a/Engine/Source/Programs/UnrealCloudDDC/Jupiter.Common/BaseStartup.cs b/Engine/Source/Programs/UnrealCloudDDC/Jupiter.Common/BaseStartup.cs
--- a/Engine/Source/Programs/UnrealCloudDDC/Jupiter.Common/BaseStartup.cs
+++ b/Engine/Source/Programs/UnrealCloudDDC/Jupiter.Common/BaseStartup.cs
@@ -12,11 +12,13 @@ using System.Reflection;
 using System.Text;
 using System.Text.Json;
 using System.Text.Json.Serialization;
+using System.Threading.Tasks;
 using Amazon;
 using EpicGames.AspNet;
 using Jupiter.Common;
 using Jupiter.Common.Implementation;
 using Microsoft.AspNetCore.Authentication;
+using Microsoft.AspNetCore.Authentication.JwtBearer;
 using Microsoft.AspNetCore.Authorization;
 using Microsoft.AspNetCore.Builder;
 using Microsoft.AspNetCore.Diagnostics.HealthChecks;
@@ -172,12 +174,30 @@ namespace Jupiter
 							});
 							break;
 						case SchemeImplementations.Okta:
+							JwtBearerEvents bearerEvents = new JwtBearerEvents();
+							if (Auth.RemapNameClaim)
+							{
+								bool firstTime = true;
+								bearerEvents.OnMessageReceived += context =>
+								{
+									if (!firstTime)
+									{
+										return Task.CompletedTask;
+									}
+
+									firstTime = false;
+									context.Options.TokenValidationParameters.NameClaimType = "http://schemas.xmlsoap.org/ws/2005/05/identity/claims/nameidentifier";
+									return Task.CompletedTask;
+								};
+							}
+							
 							availableSchemes.Add(name);
 							authenticationBuilder.AddOktaWebApi(name, new OktaWebApiOptions
 							{
 								OktaDomain = scheme.OktaDomain,
 								AuthorizationServerId = scheme.OktaAuthorizationServerId,
 								Audience = scheme.JwtAudience,
+								JwtBearerEvents = bearerEvents
 							});
 							break;
 						default:
@@ -675,6 +695,11 @@ namespace Jupiter
 		[System.Diagnostics.CodeAnalysis.SuppressMessage("Usage", "CA2227:Collection properties should be read only", Justification = "Used by the configuration system")]
 		public List<AclEntry> Acls { get; set; } = new List<AclEntry>();
 
+		/// <summary>
+		/// Remaps name claim from http://schemas.xmlsoap.org/ws/2005/05/identity/claims/name to http://schemas.xmlsoap.org/ws/2005/05/identity/claims/nameidentifier which is required for Okta
+		/// </summary>
+		public bool RemapNameClaim { get; set; } = true;
+
 		public IEnumerable<ValidationResult> Validate(ValidationContext validationContext)
 		{
 			List<ValidationResult> validationResults = new List<ValidationResult>();
diff --git a/Engine/Source/Programs/UnrealCloudDDC/Jupiter.Common/Jupiter.Common.csproj b/Engine/Source/Programs/UnrealCloudDDC/Jupiter.Common/Jupiter.Common.csproj
--- a/Engine/Source/Programs/UnrealCloudDDC/Jupiter.Common/Jupiter.Common.csproj
+++ b/Engine/Source/Programs/UnrealCloudDDC/Jupiter.Common/Jupiter.Common.csproj
@@ -42,11 +42,11 @@
     <PackageReference Include="IdentityServer4.AccessTokenValidation" Version="3.0.1" />
 
     <PackageReference Include="Microsoft.AspNet.WebApi.Client" Version="6.0.0" />
-    <PackageReference Include="Microsoft.AspNetCore.Authentication.JwtBearer" Version="8.0.1" />
+    <PackageReference Include="Microsoft.AspNetCore.Authentication.JwtBearer" Version="8.0.4" />
     <PackageReference Include="Microsoft.Extensions.Diagnostics.HealthChecks" Version="8.0.1" />
     <PackageReference Include="Microsoft.Extensions.Hosting" Version="8.0.0" />
 
-    <PackageReference Include="Okta.AspNetCore" Version="4.5.0" />
+    <PackageReference Include="Okta.AspNetCore" Version="4.6.1" />
     <!-- Adding explicit version of identity model to workaround deletion of a internal symbol used by Okta.AspNetCore -->
     <PackageReference Include="Microsoft.IdentityModel.Tokens" Version="7.2.0" />
     <PackageReference Include="Microsoft.IdentityModel.Logging" Version="7.2.0" />
@@ -56,6 +56,7 @@
 
     <!--Support reading yaml config files -->
     <PackageReference Include="NetEscapades.Configuration.Yaml" Version="3.1.0" />
+    <PackageReference Include="Serilog.Enrichers.ClientInfo" Version="2.0.3" />
     <PackageReference Include="Serilog.Extensions.Hosting" Version="8.0.0" />
     <PackageReference Include="SharpYaml" Version="2.1.0" />
 
diff --git a/Engine/Source/Programs/UnrealCloudDDC/Jupiter/Controllers/SymbolsController.cs b/Engine/Source/Programs/UnrealCloudDDC/Jupiter/Controllers/SymbolsController.cs
--- a/Engine/Source/Programs/UnrealCloudDDC/Jupiter/Controllers/SymbolsController.cs
+++ b/Engine/Source/Programs/UnrealCloudDDC/Jupiter/Controllers/SymbolsController.cs
@@ -25,7 +25,6 @@ using Microsoft.Extensions.Logging;
 using Microsoft.Extensions.Options;
 using OpenTelemetry.Trace;
 using Serilog;
-using Serilog.Context;
 using BinaryReader = System.IO.BinaryReader;
 using ILogger = Serilog.ILogger;
 
@@ -59,7 +58,7 @@ namespace Jupiter.Controllers
 			_requestHelper = requestHelper;
 			_tracer = tracer;
 			_logger = logger;
-			_auditLogger = symbolSettings.CurrentValue.EnableAuditLog ? Serilog.Log.ForContext("SymbolsAuditLog", null) : null;
+			_auditLogger = symbolSettings.CurrentValue.EnableAuditLog ? Serilog.Log.ForContext("LogType", "Audit") : null;
 		}
 
 		/// <summary>
@@ -235,7 +234,6 @@ namespace Jupiter.Controllers
 
 		private void LogAuditEntry(HttpMethod method, NamespaceId ns, string moduleName, string identifier, string fileName)
 		{
-			using IDisposable context = LogContext.PushProperty("LogType", "audit");
 			_auditLogger?.Information("{HttpMethod} '{Namespace}'/'{ModuleName}:{Identifier}' {FileName} IP:{IP} User:{Username} UserAgent:\"{Useragent}\"", method,ns, moduleName, identifier, fileName, Request.HttpContext.Connection.RemoteIpAddress, User?.Identity?.Name ?? "Unknown-user", string.Join(' ', Request.Headers.UserAgent.ToArray()));
 		}
 
diff --git a/Engine/Source/Programs/UnrealCloudDDC/Jupiter/Jupiter.csproj b/Engine/Source/Programs/UnrealCloudDDC/Jupiter/Jupiter.csproj
--- a/Engine/Source/Programs/UnrealCloudDDC/Jupiter/Jupiter.csproj
+++ b/Engine/Source/Programs/UnrealCloudDDC/Jupiter/Jupiter.csproj
@@ -80,6 +80,7 @@
     <PackageReference Include="Microsoft.Extensions.Logging.Debug" Version="8.0.0" />
     <PackageReference Include="MongoDB.Driver" Version="2.23.1" />
     <PackageReference Include="Newtonsoft.Json" Version="13.0.3" />
+    <PackageReference Include="Serilog.Enrichers.ClientInfo" Version="2.0.3" />
     <PackageReference Include="SerilogAnalyzer" Version="0.15.0" />
     <PackageReference Include="Crc32.NET" Version="1.2.0" />
     <PackageReference Include="K4os.Compression.LZ4" Version="1.3.6" />
