Cloud DDC - Decreasing min body request rate, e.g. allowing clients to send less data in a connection before it gets terminated (mostly allowed to take longer to send the data). We hope this will resolve issues that some clients see when there are a lot of parallel uploads to DDC (so some connections gets starved)

diff --git a/Engine/Source/Programs/UnrealCloudDDC/Jupiter.Common/BaseProgram.cs b/Engine/Source/Programs/UnrealCloudDDC/Jupiter.Common/BaseProgram.cs
--- a/Engine/Source/Programs/UnrealCloudDDC/Jupiter.Common/BaseProgram.cs
+++ b/Engine/Source/Programs/UnrealCloudDDC/Jupiter.Common/BaseProgram.cs
@@ -124,6 +124,9 @@ namespace Jupiter
 					{
 						options.AddServerHeader = false;
 
+						// Decrease min request body rate to be more forgiving for clients that are going wide and not sending as much data per connection
+						options.Limits.MinRequestBodyDataRate = new MinDataRate(10, TimeSpan.FromSeconds(60));
+
 						string socketsRoot = settings.DomainSocketsRoot;
 
 						if (settings.UseDomainSockets)
