Cloud DDC - Fixed issue were helm charts was not accepting the type of replicator to use always reverting to a Ref replicator.

diff --git a/Engine/Source/Programs/UnrealCloudDDC/Helm/UnrealCloudDDC/templates/configmap-worker.yaml b/Engine/Source/Programs/UnrealCloudDDC/Helm/UnrealCloudDDC/templates/configmap-worker.yaml
--- a/Engine/Source/Programs/UnrealCloudDDC/Helm/UnrealCloudDDC/templates/configmap-worker.yaml
+++ b/Engine/Source/Programs/UnrealCloudDDC/Helm/UnrealCloudDDC/templates/configmap-worker.yaml
@@ -68,6 +68,7 @@ Replication:
     MaxParallelReplications: {{ default 128  .MaxParallelReplications }} 
     SkipSnapshot: {{ default true  .SkipSnapshot }}
     PageSize: {{ default 1000  .PageSize }}
+    Version: {{ default "Refs"  .Version }}
 {{- end -}}
 {{- end }}
 
