Cloud DDC - Improvements to blob replication

Increased time between polls for replication to reduce load on the server when we are caught up, introduces a slight increase in time for a blob to be replicate when caught up but this shouldn't matter much.

Improved the metrics being reported by the blob replicator to make them more consistent when starting from very far behind (or from a clean slate)

Added check if a blob exists before attempting to replicate it.

diff --git a/Engine/Source/Programs/UnrealCloudDDC/Jupiter/Implementation/Replication/BlobsReplicator.cs b/Engine/Source/Programs/UnrealCloudDDC/Jupiter/Implementation/Replication/BlobsReplicator.cs
--- a/Engine/Source/Programs/UnrealCloudDDC/Jupiter/Implementation/Replication/BlobsReplicator.cs
+++ b/Engine/Source/Programs/UnrealCloudDDC/Jupiter/Implementation/Replication/BlobsReplicator.cs
@@ -149,7 +149,7 @@ namespace Jupiter.Implementation
 
 			_logger.LogDebug("Read Replication state for replicator: {Name}. {LastBucket} {LastEvent}.", _name, _replicationState.LastBucket, _replicationState.LastEvent);
 
-			s_replicationAttempts?.Add(1, new KeyValuePair<string, object?>("replicator", _name), new KeyValuePair<string, object?>("namespace", _namespace));
+			LogReplicationHeartbeat();
 
 			bool hasRun;
 			int countOfReplicationsDone = 0;
@@ -199,6 +199,11 @@ namespace Jupiter.Implementation
 			return hasRun;
 		}
 
+		private void LogReplicationHeartbeat()
+		{
+			s_replicationAttempts?.Add(1, new KeyValuePair<string, object?>("replicator", _name), new KeyValuePair<string, object?>("namespace", _namespace));
+		}
+
 		private async Task<HttpRequestMessage> BuildHttpRequestAsync(HttpMethod httpMethod, Uri uri)
 		{
 			string? token = await _serviceCredentials.GetTokenAsync();
@@ -228,6 +233,11 @@ namespace Jupiter.Implementation
 				return countOfReplicationsDone;
 			}
 
+			// log the metrics when we start / have nothing to do to make sure there is data in these most of the time
+			int timeBucketsBehindStart = (int)(DateTime.UtcNow - lastBucket.FromReplicationBucketIdentifier()).TotalMinutes / 5;
+			s_replicationBehindCounter?.Record(timeBucketsBehindStart, new KeyValuePair<string, object?>("replicator", _name), new KeyValuePair<string, object?>("namespace", _namespace));
+			s_replicatedCounter?.Record(countOfReplicationsDone, new KeyValuePair<string, object?>("replicator", _name), new KeyValuePair<string, object?>("namespace", _namespace));
+
 			foreach (string refBucket in GetRefBuckets(lastBucket, replicationToken))
 			{
 				if (replicationToken.IsCancellationRequested)
@@ -235,6 +245,8 @@ namespace Jupiter.Implementation
 					break;
 				}
 
+				LogReplicationHeartbeat();
+
 				await Parallel.ForEachAsync(GetBlobEventsAsync(ns, refBucket, replicationToken),
 					new ParallelOptions { MaxDegreeOfParallelism = maxParallelism, CancellationToken = linkedTokenSource.Token },
 					async (BlobReplicationLogEvent @event, CancellationToken ctx) =>
@@ -276,7 +288,7 @@ namespace Jupiter.Implementation
 
 				if (isRecentBucket)
 				{
-					_logger.LogWarning("Recent bucket for when replicating {Name}, will run replication again to ensure consistency", _name);
+					_logger.LogWarning("Reached recent bucket when replicating {Name}, will run replication again to ensure consistency as this bucket may change again", _name);
 				}
 				else
 				{
@@ -288,7 +300,6 @@ namespace Jupiter.Implementation
 				Info.LastRun = DateTime.Now;
 				int timeBucketsBehind = (int)(DateTime.UtcNow - timestamp).TotalMinutes / 5;
 				s_replicationBehindCounter?.Record(timeBucketsBehind, new KeyValuePair<string, object?>("replicator", _name), new KeyValuePair<string, object?>("namespace", _namespace));
-
 				s_replicatedCounter?.Record(countOfReplicationsDone, new KeyValuePair<string, object?>("replicator", _name), new KeyValuePair<string, object?>("namespace", _namespace));
 
 				_logger.LogInformation("{Name} replicated all events up to {Time} . Bucket: {EventBucket}", _name, timestamp, refBucket);
@@ -428,6 +439,13 @@ namespace Jupiter.Implementation
 				.SetAttribute("operation.name", "replicator.replicate_blob")
 				.SetAttribute("resource.name", $"{ns}.{blob}");
 
+			bool exists = await _blobService.ExistsAsync(ns, blob, storageLayers: null, cancellationToken);
+			if (exists)
+			{
+				_logger.LogDebug("Not replicating blob {Blob} in {Namespace} as it already existed.", blob, ns);
+				return false;
+			}
+
 			_logger.LogDebug("Attempting to replicate blob {Blob} in {Namespace}.", blob, ns);
 
 			const int RetryAttempts = 5;
diff --git a/Engine/Source/Programs/UnrealCloudDDC/Jupiter/Implementation/Replication/ReplicationSettings.cs b/Engine/Source/Programs/UnrealCloudDDC/Jupiter/Implementation/Replication/ReplicationSettings.cs
--- a/Engine/Source/Programs/UnrealCloudDDC/Jupiter/Implementation/Replication/ReplicationSettings.cs
+++ b/Engine/Source/Programs/UnrealCloudDDC/Jupiter/Implementation/Replication/ReplicationSettings.cs
@@ -19,7 +19,7 @@ namespace Jupiter
 		/// </summary>
 		[Required]
 		[Range(15, int.MaxValue)]
-		public int ReplicationPollFrequencySeconds { get; set; } = 15;
+		public int ReplicationPollFrequencySeconds { get; set; } = 60;
 
 		[Required]
 		[System.Diagnostics.CodeAnalysis.SuppressMessage("Usage", "CA2227:Collection properties should be read only", Justification = "Used by serialization")]
