Cloud DDC - Added warning when we fail to cleanup the stats tracking and attempted to make it a bit more resilient by at least cleaning up the ref tracking even if we are unable to handle the blobs. Hopefully this helps us track down why the stats tracking seems to be more and more wrong as times goes on.

diff --git a/Engine/Source/Programs/UnrealCloudDDC/Jupiter/Implementation/GC/CacheCleanup.cs b/Engine/Source/Programs/UnrealCloudDDC/Jupiter/Implementation/GC/CacheCleanup.cs
--- a/Engine/Source/Programs/UnrealCloudDDC/Jupiter/Implementation/GC/CacheCleanup.cs
+++ b/Engine/Source/Programs/UnrealCloudDDC/Jupiter/Implementation/GC/CacheCleanup.cs
@@ -170,15 +170,17 @@ namespace Jupiter.Implementation
 			{
 				if (_cloudDDCSettings.CurrentValue.EnableBucketStatsTracking)
 				{
+					List<BlobId> blobs = new List<BlobId>();
 					try
 					{
-						List<BlobId> blobs = await _objectService.GetReferencedBlobsAsync(ns, bucket, name, ignoreMissingBlobs: true, cancellationToken: cancellationToken);
-						await _blobIndex.RemoveBlobFromBucketListAsync(ns, bucket, name, blobs, cancellationToken);
+						blobs = await _objectService.GetReferencedBlobsAsync(ns, bucket, name, ignoreMissingBlobs: true, cancellationToken: cancellationToken);
 					}
 					catch (RefNotFoundException)
 					{
 						// if the ref is already deleted its not possible for us to cleanup the bucket list
+						_logger.LogWarning("Ref {Name} {Bucket} {Namespace} not found when cleaning up stats tracking, unable to remove blob tracking thus stats will be incorrect.", name, bucket, ns);
 					}
+					await _blobIndex.RemoveBlobFromBucketListAsync(ns, bucket, name, blobs, cancellationToken);
 				}
 
 				storeDelete = await _referencesStore.DeleteAsync(ns, bucket, name, cancellationToken);
