Cloud DDC - Added ability to configure a HPA (Horizontal pod autoscaler) - this can create more pods if the current cpu load is to high. Currently experimental until we have had time to tweak the configuration of it, is intended to be used to autoscale in case of extreme load changes (other DDC layers going down so all traffic hits Cloud DDC).

Should be combined with a node autoscaler which is not configured in the helm chart.

Also removed some attributed specified in the security context that makes no sense.

diff --git a/Engine/Source/Programs/UnrealCloudDDC/Helm/UnrealCloudDDC/templates/hpa.yaml b/Engine/Source/Programs/UnrealCloudDDC/Helm/UnrealCloudDDC/templates/hpa.yaml
new file mode 100644
--- /dev/null
+++ b/Engine/Source/Programs/UnrealCloudDDC/Helm/UnrealCloudDDC/templates/hpa.yaml
@@ -0,0 +1,43 @@
+{{- if .Values.autoscaling.hpa.enabled  }}
+apiVersion: {{ include "common.capabilities.hpa.apiVersion" ( dict "context" $ ) }}
+kind: HorizontalPodAutoscaler
+metadata:
+  name: {{ include "common.names.fullname" . }}
+  namespace: {{ include "common.names.namespace" . | quote }}
+  labels: {{- include "common.labels.standard" . | nindent 4 }}
+  {{- if or .Values.autoscaling.hpa.annotations .Values.commonAnnotations }}
+  {{- $annotations := include "common.tplvalues.merge" ( dict "values" ( list .Values.autoscaling.hpa.annotations .Values.commonAnnotations ) "context" . ) }}
+  annotations: {{- include "common.tplvalues.render" ( dict "value" $annotations "context" $) | nindent 4 }}
+  {{- end }}
+spec:
+  {{- if .Values.autoscaling.hpa.behavior }}
+  behavior:
+    {{- toYaml .Values.autoscaling.hpa.behavior | nindent 4 }}
+  {{- end }}
+  scaleTargetRef:
+    apiVersion: apps/v1
+    kind: Deployment
+    name: {{ include "common.names.fullname" . }}
+  minReplicas: {{ .Values.autoscaling.hpa.minReplicas }}
+  maxReplicas: {{ .Values.autoscaling.hpa.maxReplicas }}
+  metrics:
+  {{- if .Values.autoscaling.hpa.targetCPUUtilizationPercentage }}
+    - type: Resource
+      resource:
+        name: cpu
+        target:
+          averageUtilization: {{ .Values.autoscaling.hpa.targetCPUUtilizationPercentage }}
+          type: Utilization
+  {{- end }}
+  {{- if .Values.autoscaling.hpa.targetMemoryUtilizationPercentage }}
+    - type: Resource
+      resource:
+        name: memory
+        target:
+          averageUtilization: {{ .Values.autoscaling.hpa.targetMemoryUtilizationPercentage }}
+          type: Utilization
+  {{- end }}
+  {{- if .Values.autoscaling.hpa.customRules -}}
+    {{- toYaml .Values.autoscaling.hpa.customRules | nindent 4}}
+  {{- end -}}
+{{- end }}
\ No newline at end of file
diff --git a/Engine/Source/Programs/UnrealCloudDDC/Helm/UnrealCloudDDC/values.yaml b/Engine/Source/Programs/UnrealCloudDDC/Helm/UnrealCloudDDC/values.yaml
--- a/Engine/Source/Programs/UnrealCloudDDC/Helm/UnrealCloudDDC/values.yaml
+++ b/Engine/Source/Programs/UnrealCloudDDC/Helm/UnrealCloudDDC/values.yaml
@@ -42,8 +42,7 @@ env: []
 ## Security Context
 ## ref: https://kubernetes.io/docs/tasks/configure-pod-container/security-context/
 ##
-securityContext:
-  enabled: false
+securityContext: {}
 
 ## Pod annotations
 ## ref: https://kubernetes.io/docs/concepts/overview/working-with-objects/annotations/
@@ -170,6 +169,37 @@ persistence:
   ##
   annotations: {}
 
+
+autoscaling:
+  ## Horizontal Pod Autoscaler configuration
+  ## ref: https://kubernetes.io/docs/tasks/run-application/horizontal-pod-autoscale/
+  ##
+  hpa:
+    ## Annotations for HPA resource
+    ##
+    annotations: {}
+    ## Enable HPA
+    ##
+    enabled: false
+    ## Min replicas
+    ##
+    minReplicas: 2
+    ## Max replicas
+    ##
+    maxReplicas: 4
+    ## Target CPU utilization percentage
+    ##
+    targetCPUUtilizationPercentage: 75
+    ## Target Memory utilization percentage
+    ##
+    targetMemoryUtilizationPercentage: ""
+    ## Custom rules
+    ##
+    customRules: []
+    ## HPA Behavior
+    ##
+    behavior: {}
+
 ## Service properties
 ##
 service:
