Cloud DDC - Added ability to specify requests for nginx pod and added some requests to the dotnet sidecar, requirements to be able to enable the autoscaler.

diff --git a/Engine/Source/Programs/UnrealCloudDDC/Helm/UnrealCloudDDC/templates/deployment.yaml b/Engine/Source/Programs/UnrealCloudDDC/Helm/UnrealCloudDDC/templates/deployment.yaml
--- a/Engine/Source/Programs/UnrealCloudDDC/Helm/UnrealCloudDDC/templates/deployment.yaml
+++ b/Engine/Source/Programs/UnrealCloudDDC/Helm/UnrealCloudDDC/templates/deployment.yaml
@@ -163,6 +163,9 @@ spec:
         volumeMounts:
           - name: diagnostics
             mountPath: /tmp
+        resources:
+          requests: 
+            cpu: 50m
     {{- end }}
     {{- if .Values.nginx.enabled }}
       - name: nginx
@@ -214,6 +217,9 @@ spec:
               - "sleep"
               - "{{ .Values.preStopHookDelay }}"
         {{- end }}
+        {{- if .Values.nginx.resources }}
+        resources: {{- toYaml .Values.nginx.resources | nindent 12 }}
+        {{- end }}
         volumeMounts:
           - name: nginx
             mountPath: /nginx
