Cloud DDC - Support the renamed ue-session headers (now called x-ue-session to indicate that its not a standard http header). Will still accept and handle the old values for older not yet upgraded clients.

diff --git a/Engine/Source/Programs/UnrealCloudDDC/Jupiter.Common/BaseProgram.cs b/Engine/Source/Programs/UnrealCloudDDC/Jupiter.Common/BaseProgram.cs
--- a/Engine/Source/Programs/UnrealCloudDDC/Jupiter.Common/BaseProgram.cs
+++ b/Engine/Source/Programs/UnrealCloudDDC/Jupiter.Common/BaseProgram.cs
@@ -52,6 +52,7 @@ namespace Jupiter
 			Log.Logger = new LoggerConfiguration()
 				.ReadFrom.Configuration(Configuration)
 				.Enrich.With<DatadogLogEnricher>()
+				.Enrich.WithRequestHeader("x-ue-session", "ue-session")
 				.Enrich.WithRequestHeader("ue-session", "ue-session")
 				.CreateLogger();
 
diff --git a/Engine/Source/Programs/UnrealCloudDDC/Jupiter.Common/BaseStartup.cs b/Engine/Source/Programs/UnrealCloudDDC/Jupiter.Common/BaseStartup.cs
--- a/Engine/Source/Programs/UnrealCloudDDC/Jupiter.Common/BaseStartup.cs
+++ b/Engine/Source/Programs/UnrealCloudDDC/Jupiter.Common/BaseStartup.cs
@@ -301,7 +301,7 @@ namespace Jupiter
 				{
 					options.EnrichWithHttpRequest = (activity, request) =>
 					{
-						if (request.Headers.TryGetValue("ue-session", out StringValues ueSessionValues))
+						if (request.Headers.TryGetValue("x-ue-session", out StringValues ueSessionValues))
 						{
 							if (ueSessionValues.Count != 0)
 							{
@@ -309,13 +309,29 @@ namespace Jupiter
 							}
 						}
 
-						if (request.Headers.TryGetValue("ue-request", out StringValues ueRequestValues))
+						if (request.Headers.TryGetValue("ue-session", out StringValues ueSessionValuesOld))
+						{
+							if (ueSessionValuesOld.Count != 0)
+							{
+								activity.AddTag("ue-session", ueSessionValuesOld.First());
+							}
+						}
+
+						if (request.Headers.TryGetValue("x-ue-request", out StringValues ueRequestValues))
 						{
 							if (ueRequestValues.Count != 0)
 							{
 								activity.AddTag("ue-request", ueRequestValues.First());
 							}
 						}
+
+						if (request.Headers.TryGetValue("ue-request", out StringValues ueRequestValuesOld))
+						{
+							if (ueRequestValuesOld.Count != 0)
+							{
+								activity.AddTag("ue-request", ueRequestValuesOld.First());
+							}
+						}
 					};
 				});
 				builder.SetSampler(new AlwaysOnSampler());
