Cloud DDC - Added more information to tracing for exists remote operations to see if this might be the reason we see occasional issues with the oplog pulls.

diff --git a/Engine/Source/Programs/UnrealCloudDDC/Jupiter/Implementation/Blob/BlobService.cs b/Engine/Source/Programs/UnrealCloudDDC/Jupiter/Implementation/Blob/BlobService.cs
--- a/Engine/Source/Programs/UnrealCloudDDC/Jupiter/Implementation/Blob/BlobService.cs
+++ b/Engine/Source/Programs/UnrealCloudDDC/Jupiter/Implementation/Blob/BlobService.cs
@@ -746,7 +746,9 @@ public class BlobService : IBlobService
 	public async Task<bool> ExistsInRemoteAsync(NamespaceId ns, BlobId blob, CancellationToken cancellationToken)
 	{
 		IServerTiming? serverTiming = _httpContextAccessor.HttpContext?.RequestServices.GetService<IServerTiming>();
-		using TelemetrySpan scope = _tracer.StartActiveSpan("HierarchicalStore.ExistsRemote").SetAttribute("operation.name", "HierarchicalStore.ExistsRemote");
+		using TelemetrySpan scope = _tracer.StartActiveSpan("HierarchicalStore.ExistsRemote")
+			.SetAttribute("operation.name", "HierarchicalStore.ExistsRemote")
+			.SetAttribute("resource.name", blob.ToString());
 
 		using ServerTimingMetricScoped? serverTimingScope = serverTiming?.CreateServerTimingMetricScope("blob.exists-remote", "Verify if blob exists in remotes");
 
@@ -758,14 +760,17 @@ public class BlobService : IBlobService
 		// if it exists in more then one region, we are sure it exists somewhere that is not here
 		if (regions.Count > 1)
 		{
+			scope.SetAttribute("ObjectFound", true.ToString());
 			return true;
 		}
 
 		if (regions.Any(region => !string.Equals(region, _currentSite, StringComparison.OrdinalIgnoreCase)))
 		{
+			scope.SetAttribute("ObjectFound", true.ToString());
 			return true;
 		}
 
+		scope.SetAttribute("ObjectFound", false.ToString());
 		return false;
 	}
 
