Cloud DDC - Added options to enable the debug endpoints which are now disabled by default, and added max concurrent calls to it to reduce risk with them.

diff --git a/Engine/Source/Programs/UnrealCloudDDC/Helm/UnrealCloudDDC/templates/configmap.yaml b/Engine/Source/Programs/UnrealCloudDDC/Helm/UnrealCloudDDC/templates/configmap.yaml
--- a/Engine/Source/Programs/UnrealCloudDDC/Helm/UnrealCloudDDC/templates/configmap.yaml
+++ b/Engine/Source/Programs/UnrealCloudDDC/Helm/UnrealCloudDDC/templates/configmap.yaml
@@ -76,6 +76,11 @@ Jupiter:
 {{ $config.Jupiter | toYaml | nindent 2 }}
 {{- end }}
 
+{{- if $config.Debug }}
+Debug: 
+{{ $config.Debug | toYaml | nindent 2 }}
+{{- end }}
+
 {{ if .Values.nginx.useDomainSockets }}
 Jupiter:
   UseDomainSockets: true
diff --git a/Engine/Source/Programs/UnrealCloudDDC/Jupiter/Controllers/DebugController.cs b/Engine/Source/Programs/UnrealCloudDDC/Jupiter/Controllers/DebugController.cs
--- a/Engine/Source/Programs/UnrealCloudDDC/Jupiter/Controllers/DebugController.cs
+++ b/Engine/Source/Programs/UnrealCloudDDC/Jupiter/Controllers/DebugController.cs
@@ -1,8 +1,11 @@
 ﻿// Copyright Epic Games, Inc. All Rights Reserved.
 
 using System;
+using System.Net;
+using System.Threading;
 using Microsoft.AspNetCore.Authorization;
 using Microsoft.AspNetCore.Mvc;
+using Microsoft.Extensions.Options;
 
 namespace Jupiter.Controllers
 {
@@ -10,6 +13,13 @@ namespace Jupiter.Controllers
 	[Route("api/v1/c/_debug")]
 	public class DebugController : Controller
 	{
+		private readonly IOptionsMonitor<DebugSettings> _settings;
+		private int _currentConcurrentCalls;
+
+		public DebugController(IOptionsMonitor<DebugSettings> settings)
+		{
+			_settings = settings;
+		}
 		/// <summary>
 		/// Return bytes of specified length with auth, used for testing only
 		/// </summary>
@@ -18,7 +28,25 @@ namespace Jupiter.Controllers
 		[Authorize]
 		public IActionResult GetBytes([FromQuery] int length = 1)
 		{
-			return GenerateByteResponse(length);
+			if (!_settings.CurrentValue.EnableDebugEndpoints)
+			{
+				return NotFound();
+			}
+
+			try
+			{
+				int concurrentCalls = Interlocked.Increment(ref _currentConcurrentCalls);
+				if (concurrentCalls > _settings.CurrentValue.MaxConcurrentCalls)
+				{
+					return StatusCode((int)HttpStatusCode.TooManyRequests);
+				}
+
+				return GenerateByteResponse(length);
+			}
+			finally
+			{
+				Interlocked.Decrement(ref _currentConcurrentCalls);
+			}	
 		}
 
 		/// <summary>
@@ -28,7 +56,24 @@ namespace Jupiter.Controllers
 		[HttpGet("getBytesWithoutAuth")]
 		public IActionResult GetBytesWithoutAuth([FromQuery] int length = 1)
 		{
-			return GenerateByteResponse(length);
+			if (!_settings.CurrentValue.EnableDebugEndpoints)
+			{
+				return NotFound();
+			}
+
+			try
+			{
+				int concurrentCalls = Interlocked.Increment(ref _currentConcurrentCalls);
+				if (concurrentCalls > _settings.CurrentValue.MaxConcurrentCalls)
+				{
+					return StatusCode((int)HttpStatusCode.TooManyRequests);
+				}
+				return GenerateByteResponse(length);
+			}
+			finally
+			{
+				Interlocked.Decrement(ref _currentConcurrentCalls);
+			}
 		}
 
 		private FileContentResult GenerateByteResponse(int length)
@@ -38,4 +83,10 @@ namespace Jupiter.Controllers
 			return File(generatedData, "application/octet-stream");
 		}
 	}
+
+	public class DebugSettings
+	{
+		public bool EnableDebugEndpoints { get; set; } = false;
+		public int MaxConcurrentCalls { get; set; } = 10;
+	}
 }
diff --git a/Engine/Source/Programs/UnrealCloudDDC/Jupiter/Jupiter.Startup.cs b/Engine/Source/Programs/UnrealCloudDDC/Jupiter/Jupiter.Startup.cs
--- a/Engine/Source/Programs/UnrealCloudDDC/Jupiter/Jupiter.Startup.cs
+++ b/Engine/Source/Programs/UnrealCloudDDC/Jupiter/Jupiter.Startup.cs
@@ -91,6 +91,7 @@ namespace Jupiter
 			services.AddOptions<S3Settings>().Bind(Configuration.GetSection("S3")).ValidateDataAnnotations();
 			services.AddOptions<AzureSettings>().Bind(Configuration.GetSection("Azure")).ValidateDataAnnotations();
 			services.AddOptions<FilesystemSettings>().Bind(Configuration.GetSection("Filesystem")).ValidateDataAnnotations();
+			services.AddOptions<DebugSettings>().Bind(Configuration.GetSection("Debug")).ValidateDataAnnotations();
 
 			services.AddOptions<NginxSettings>().Bind(Configuration.GetSection("Nginx")).ValidateDataAnnotations();
 
