Cloud DDC - Added ability to limit which namespaces a auth scheme is allowed to grant access to. Usually you would want to use the claims to drive this but this can be useful if you have a 3rd party auth provider that you do not fully control that you want to make sure can only access certain namespaces.

diff --git a/Engine/Source/Programs/UnrealCloudDDC/Jupiter.Common/BaseStartup.cs b/Engine/Source/Programs/UnrealCloudDDC/Jupiter.Common/BaseStartup.cs
--- a/Engine/Source/Programs/UnrealCloudDDC/Jupiter.Common/BaseStartup.cs
+++ b/Engine/Source/Programs/UnrealCloudDDC/Jupiter.Common/BaseStartup.cs
@@ -650,6 +650,11 @@ namespace Jupiter
 		[Required]
 		public string JwtAudience { get; set; } = "";
 
+		/// <summary>
+		/// The namespaces which these scheme is allowed to grant access to, all if this is omitted or empty
+		/// </summary>
+		public string[] AllowedNamespaces { get; set; } = Array.Empty<string>();
+
 		public IEnumerable<ValidationResult> Validate(ValidationContext validationContext)
 		{
 			List<ValidationResult> validationResults = new List<ValidationResult>();
diff --git a/Engine/Source/Programs/UnrealCloudDDC/Jupiter/Controllers/RequestHelper.cs b/Engine/Source/Programs/UnrealCloudDDC/Jupiter/Controllers/RequestHelper.cs
--- a/Engine/Source/Programs/UnrealCloudDDC/Jupiter/Controllers/RequestHelper.cs
+++ b/Engine/Source/Programs/UnrealCloudDDC/Jupiter/Controllers/RequestHelper.cs
@@ -1,6 +1,7 @@
 // Copyright Epic Games, Inc. All Rights Reserved.
 
 using System;
+using System.Linq;
 using System.Security.Claims;
 using System.Threading.Tasks;
 using EpicGames.Horde.Storage;
@@ -19,13 +20,15 @@ public class RequestHelper : IRequestHelper
 	private readonly IAuthorizationService _authorizationService;
 	private readonly INamespacePolicyResolver _namespacePolicyResolver;
 	private readonly IOptionsMonitor<JupiterSettings> _settings;
+	private readonly IOptionsMonitor<AuthSettings> _authSettings;
 	private readonly Tracer _tracer;
 
-	public RequestHelper(IAuthorizationService authorizationService, INamespacePolicyResolver namespacePolicyResolver, IOptionsMonitor<JupiterSettings> settings, Tracer tracer)
+	public RequestHelper(IAuthorizationService authorizationService, INamespacePolicyResolver namespacePolicyResolver, IOptionsMonitor<JupiterSettings> settings, IOptionsMonitor<AuthSettings> authSettings, Tracer tracer)
 	{
 		_authorizationService = authorizationService;
 		_namespacePolicyResolver = namespacePolicyResolver;
 		_settings = settings;
+		_authSettings = authSettings;
 		_tracer = tracer;
 	}
 
@@ -43,6 +46,19 @@ public class RequestHelper : IRequestHelper
 			return new ForbidResult();
 		}
 
+		// fetch the value of the issuer claim
+		string? issuer = user.FindFirstValue("iss");
+		AuthSchemeEntry? authScheme = _authSettings.CurrentValue.Schemes.Values.FirstOrDefault(entry => entry.JwtAuthority == issuer);
+		if (authScheme?.AllowedNamespaces.Length != 0)
+		{
+			// check if the auth scheme is allowed to grant access to this namespace
+			if (!authScheme!.AllowedNamespaces.Contains(ns.ToString(), StringComparer.InvariantCultureIgnoreCase))
+			{
+				// not allowed to grant access to the namespace
+				return new ForbidResult();
+			}
+		}
+
 		bool isPublicNamespace = _namespacePolicyResolver.GetPoliciesForNs(ns).IsPublicNamespace;
 
 		// public namespaces are always accessible
