Cloud DDC - Fixed failing tests

diff --git a/Engine/Source/Programs/UnrealCloudDDC/Jupiter/Controllers/RequestHelper.cs b/Engine/Source/Programs/UnrealCloudDDC/Jupiter/Controllers/RequestHelper.cs
--- a/Engine/Source/Programs/UnrealCloudDDC/Jupiter/Controllers/RequestHelper.cs
+++ b/Engine/Source/Programs/UnrealCloudDDC/Jupiter/Controllers/RequestHelper.cs
@@ -49,13 +49,16 @@ public class RequestHelper : IRequestHelper
 		// fetch the value of the issuer claim
 		string? issuer = user.FindFirstValue("iss");
 		AuthSchemeEntry? authScheme = _authSettings.CurrentValue.Schemes.Values.FirstOrDefault(entry => entry.JwtAuthority == issuer);
-		if (authScheme?.AllowedNamespaces.Length != 0)
+		if (authScheme != null)
 		{
-			// check if the auth scheme is allowed to grant access to this namespace
-			if (!authScheme!.AllowedNamespaces.Contains(ns.ToString(), StringComparer.InvariantCultureIgnoreCase))
+			if (authScheme.AllowedNamespaces.Length != 0)
 			{
-				// not allowed to grant access to the namespace
-				return new ForbidResult();
+				// check if the auth scheme is allowed to grant access to this namespace
+				if (!authScheme.AllowedNamespaces.Contains(ns.ToString(), StringComparer.InvariantCultureIgnoreCase))
+				{
+					// not allowed to grant access to the namespace
+					return new ForbidResult();
+				}
 			}
 		}
 
