Cloud DDC - Updated the changelog and prepared release 1.2.0
Also updated docs on how to configure oplog namespaces.

diff --git a/Engine/Source/Programs/UnrealCloudDDC/Helm/UnrealCloudDDC/example-values-base.yaml b/Engine/Source/Programs/UnrealCloudDDC/Helm/UnrealCloudDDC/example-values-base.yaml
--- a/Engine/Source/Programs/UnrealCloudDDC/Helm/UnrealCloudDDC/example-values-base.yaml
+++ b/Engine/Source/Programs/UnrealCloudDDC/Helm/UnrealCloudDDC/example-values-base.yaml
@@ -78,4 +78,17 @@ global:
       #    - ReadObject
       #    - WriteObject
       #    claims: 
-      #    - group=app-ddc-storage-user
\ No newline at end of file
+      #    - group=app-ddc-storage-user
+
+      # This is an example of a namespace used for oplog storage
+      #ue.oplog:
+      #  acls:
+      #  - actions: 
+      #    - ReadObject
+      #    - WriteObject
+      #    claims: 
+      #    - group=app-ddc-storage-user
+      #  storagePool: "oplog"
+      #  UseBlobIndexForExists: true # Enabling blob index usage witin oplogs as a experiment
+      #  GcMethod: TTL # use TTL (time to live) GC instead of last access tracking
+      #  OnDemandReplication: true # replicate blobs on demand
\ No newline at end of file
diff --git a/Engine/Source/Programs/UnrealCloudDDC/Helm/UnrealCloudDDC/version.yaml b/Engine/Source/Programs/UnrealCloudDDC/Helm/UnrealCloudDDC/version.yaml
--- a/Engine/Source/Programs/UnrealCloudDDC/Helm/UnrealCloudDDC/version.yaml
+++ b/Engine/Source/Programs/UnrealCloudDDC/Helm/UnrealCloudDDC/version.yaml
@@ -1 +1 @@
-version: 1.1.1
\ No newline at end of file
+version: 1.2.0
\ No newline at end of file
diff --git a/Engine/Source/Programs/UnrealCloudDDC/README.md b/Engine/Source/Programs/UnrealCloudDDC/README.md
--- a/Engine/Source/Programs/UnrealCloudDDC/README.md
+++ b/Engine/Source/Programs/UnrealCloudDDC/README.md
@@ -26,6 +26,9 @@ Unreal Cloud DDC can signficantly help teams speed up their cook processes in th
 - [Authentication](#authentication)
   - [IdentityProvider setup](#identityprovider-setup)
   - [Namespace access](#namespace-access)
+- [Common Configuration](#common-configuration)
+  - [DDC](#ddc)
+  - [Oplogs](#oplogs)
 - [Networking setup](#networking-setup)
   - [Public Port](#public-port)
   - [Private port](#private-port)
@@ -211,6 +214,47 @@ namespace:
 If you specify multiple claims in the claims array these are ANDed together thus requiring all the claims. A claim statement like `A=B` requires claim `A` to have value `B` (or contain value `B` in the case of a array).
 You can also specify the `*` claim which grants access for any valid token no matter which claims it has, this is mostly for debug / testing scenarios and shouldnt be used for production data.
 
+# Common Configuration
+Unreal Cloud DDC can be used for multiple types of object storage and the configuration for these different kinds can vary a bit. This section will attempt to document some of the use cases and how to configure it.
+Epic applies a naming convention to namespaces that can be useful to follow, `<project>.<type>` (all lowercase), so a namespace for UE projects would be named `ue.sddc` (sddc stands for structured ddc) and a oplog namespace would be called `ue.oplog`.
+
+Namespace will share storage amongst all namespaces that share the same storage pool and they all default to the default storage pool (no suffix on bucket names). Its recommended to seperate out `oplog` and `ddc` into seperate storage pools as there storage requirments are somewhat different and this makes it easier to get an overview of how much data you are storing of each type. 
+
+
+## DDC
+For DDC there is no need to customize any options in the namespace, everything defaults to a setup that is useful for ddc.
+
+## Oplogs
+For oplogs we recommend enabling certain options on your namespace. The most important is the `GCMethod` which should be set to TTL (Time to live). This will make oplog entries automatically expire after 14 days (can be overriden by setting the `DefaultTTL`). Note that the uploaded content isn't removed after those 14 days, just the reference to it. Blob garbage collection will remove unused blobs once this happens. This makes sure old oplogs are automatically removed when not needed anymore (unlike the default GC method which keep anything that is used around which can quickly become a lot of storage for these large oplogs).
+The `storagePool` is something we recommend overriding to a seperate pool as well. This will expect you to have a S3 bucket / Azure blob storage that matched your configured name but with the storage pool suffix. So if your s3 bucket name is set to `ue-cloud-ddc-bucket` then the expected bucket for oplogs is `ue-cloud-ddc-bucket-oplog` e.g a `-oplog` suffix.
+
+Next you can enable the `UseBlobIndexForExists` option, this will use the metadata db (ScyllaDB) for exist checks rather then querying blob storage. This is much faster but puts more effort on your DB so make sure to have it scaled appopriately, Scylla usually handles this just fine. For large oplogs the existence checks can be a bottle neck.
+Lastly we have the `OnDemandReplication` option. When this is set that means any blob that is missing in a region is fetched on demand to other remote regions. This means that you do not nesscarilly need to have the speculative blob replication enabled but will be much slower for the first user that pulls down a oplog. We have that oplogs generally work okay with such a setup as long as you have enough users in a region. For smaller regions the speculative replication is better and it is always a net benefit to user experience (at a cost to transfer the data and store it).
+
+This is an example of how the policies for a namespace could be configured for oplogs.
+```
+namespace:
+  policies:
+    ue.ddc:
+      acls:
+      - actions: 
+        - ReadObject
+        - WriteObject
+        claims: 
+        - ExampleClaim   
+    ue.oplog:
+      acls:
+      - actions: 
+        - ReadObject
+        - WriteObject
+        claims: 
+        - ExampleClaim   
+      storagePool: "oplog"
+      UseBlobIndexForExists: true
+      GcMethod: TTL
+      OnDemandReplication: true
+```
+
 # Networking setup
 UnrealCloudDDC derives alot of its performance from using your normal internet connection and not relying on a VPN tunnel. As its we strongly recommend that you expose UnrealCloudDDC on a public internet endpoint. From that follows recommendations like using https and setting up your authentication (see `Authentication`) to prevent anyone from accessing this data.
 
diff --git a/Engine/Source/Programs/UnrealCloudDDC/changelog.md b/Engine/Source/Programs/UnrealCloudDDC/changelog.md
--- a/Engine/Source/Programs/UnrealCloudDDC/changelog.md
+++ b/Engine/Source/Programs/UnrealCloudDDC/changelog.md
@@ -1,5 +1,13 @@
 # Unreleased
+
+# 1.2.0
 * `InlineMaxBlobSize` option has moved from `Scylla` to `UnrealCloudDDC` as this can now be used to generally control if blob are inlined into the ref store (minimal practical difference as only Scylla supports inlining blobs). Resolves issue with `EnablePutRefBodyIntoBlobStore` option which will prevent these inlined blobs from also being added to the blob store.
+* Added experimental option that is disabled by default to allow enumeration of buckets (enabled per namespace). Could potentially be used by oplogs.
+* Optimizations for deleting blobs from multiple namespaces at once.
+* Blob replication method added, can be enabled by setting the `Version` option under each replicator to `Blobs` - this will likely be made the default in future releases. The blob based replication is faster and more reliable then our old speculative replication as this has a more explicit list of which blobs to replicate rather then the previous method that was more indirect via the submitted ref (which can be mutated). This is also more useful for oplogs where a ref has a lot of references but few changed blobs each upload.
+* Added support for new blob endpoint under refs as used in UE 5.6
+* Added configuration for HPA (Horizontal pod autoscaler) in Helm chart. We use this to autoscale Cloud DDC instance a bit during bursty periods (in combination with a node autoscaler).
+* Added option `AllowedNamespaces` on a authentication scheme to limit which namespaces a scheme is allowed to grant access to. This can be useful if you have a 3rd party authentication server you want to use but do not control and only want to grant access to some data for that party. For most use cases claims should be sufficent to control the access you need to grant.
 
 # 1.1.1
 * Fixes to helm chart when using ServiceAccounts for authentication and configuring replication
diff --git a/Engine/Source/Programs/UnrealCloudDDC/version.yaml b/Engine/Source/Programs/UnrealCloudDDC/version.yaml
--- a/Engine/Source/Programs/UnrealCloudDDC/version.yaml
+++ b/Engine/Source/Programs/UnrealCloudDDC/version.yaml
@@ -1 +1 @@
-version: 1.1.1
\ No newline at end of file
+version: 1.2.0
\ No newline at end of file
