Cloud DDC - Fixed issue with helm state always diffing causing ArgoCD to apply helm over and over again.

diff --git a/Engine/Source/Programs/UnrealCloudDDC/Helm/UnrealCloudDDC/templates/deployment.yaml b/Engine/Source/Programs/UnrealCloudDDC/Helm/UnrealCloudDDC/templates/deployment.yaml
--- a/Engine/Source/Programs/UnrealCloudDDC/Helm/UnrealCloudDDC/templates/deployment.yaml
+++ b/Engine/Source/Programs/UnrealCloudDDC/Helm/UnrealCloudDDC/templates/deployment.yaml
@@ -88,7 +88,7 @@ spec:
         {{- if .Values.env }}
         {{- include "common.tplvalues.render" ( dict "value" .Values.env "context" $ ) | nindent 8 }}
         {{- end }}
-        {{- if .Values.securityContext }}
+        {{- if .Values.securityContext.enabled }}
         securityContext: {{- include "common.tplvalues.render" (dict "value" .Values.securityContext "context" $) | nindent 8 }}
         {{- end }}
         ports:
