Cloud DDC - Fixed issue were status controllers didnt support requesting compact binary responses on its endpoints.

diff --git a/Engine/Source/Programs/UnrealCloudDDC/Jupiter/Controllers/StatusController.cs b/Engine/Source/Programs/UnrealCloudDDC/Jupiter/Controllers/StatusController.cs
--- a/Engine/Source/Programs/UnrealCloudDDC/Jupiter/Controllers/StatusController.cs
+++ b/Engine/Source/Programs/UnrealCloudDDC/Jupiter/Controllers/StatusController.cs
@@ -5,6 +5,7 @@ using System.Collections.Generic;
 using System.Linq;
 using System.Reflection;
 using System.Text.Json.Serialization;
+using EpicGames.Serialization;
 using Jupiter.Implementation;
 using Microsoft.AspNetCore.Authorization;
 using Microsoft.AspNetCore.Mvc;
@@ -85,6 +86,11 @@ namespace Jupiter.Controllers
 
 	public class PeersResponse
 	{
+		public PeersResponse()
+		{
+
+		}
+
 		[JsonConstructor]
 		public PeersResponse(string currentSite, List<KnownPeer> peers)
 		{
@@ -98,14 +104,23 @@ namespace Jupiter.Controllers
 			Peers = clusterSettings.CurrentValue.Peers.Select(settings => new KnownPeer(settings, includeInternalEndpoints, peerStatusService)).ToList();
 		}
 
+		[CbField("currentSite")]
 		public string CurrentSite { get; set; } = null!;
 
+		[CbField("peers")]
 		[System.Diagnostics.CodeAnalysis.SuppressMessage("Usage", "CA2227:Collection properties should be read only", Justification = "Used by serialization")]
 		public List<KnownPeer> Peers { get; set; } = new List<KnownPeer>();
 	}
 
 	public class KnownPeer
 	{
+		public KnownPeer()
+		{
+			Site = null!;
+			FullName = null!;
+			Endpoints = null!;
+		}
+
 		[JsonConstructor]
 		public KnownPeer(string site, string fullName, List<Uri> endpoints, int latency)
 		{
@@ -134,17 +149,30 @@ namespace Jupiter.Controllers
 			}
 		}
 
+		[CbField("site")]
 		public string Site { get; set; }
+
+		[CbField("fullName")]
 		public string FullName { get; set; }
 
+		[CbField("latency")]
 		public int Latency { get; set; }
 
 		[System.Diagnostics.CodeAnalysis.SuppressMessage("Usage", "CA2227:Collection properties should be read only", Justification = "Used by serialization")]
+		[CbField("endpoints")]
 		public List<Uri> Endpoints { get; set; }
 	}
 
 	public class StatusResponse
 	{
+		public StatusResponse()
+		{
+			Version = null!;
+			GitHash = null!;
+			Capabilities = Array.Empty<string>();
+			SiteIdentifier = null!;
+		}
+
 		public StatusResponse(string version, string gitHash, string[] capabilities, string siteIdentifier)
 		{
 			Version = version;
