Cloud DDC - Include information on when the ref was created in the replication consistency check message so that we can determine if the issues we are seeing are due to new refs or old ones. (new ones would be expected to take a bit of time to replicate)

diff --git a/Engine/Source/Programs/UnrealCloudDDC/Jupiter/Implementation/Blob/RefStoreConsistencyCheckService.cs b/Engine/Source/Programs/UnrealCloudDDC/Jupiter/Implementation/Blob/RefStoreConsistencyCheckService.cs
--- a/Engine/Source/Programs/UnrealCloudDDC/Jupiter/Implementation/Blob/RefStoreConsistencyCheckService.cs
+++ b/Engine/Source/Programs/UnrealCloudDDC/Jupiter/Implementation/Blob/RefStoreConsistencyCheckService.cs
@@ -38,7 +38,7 @@ namespace Jupiter.Implementation
 			return _settings.CurrentValue.EnableRefStoreChecks;
 		}
 
-		public RefStoreConsistencyCheckService(IOptionsMonitor<ConsistencyCheckSettings> settings, IOptionsMonitor<UnrealCloudDDCSettings> unrealCloudDDCSettings, IServiceProvider provider, ILeaderElection leaderElection, IRefService refService, IBlobIndex blobIndex, Tracer tracer, ILogger<BlobStoreConsistencyCheckService> logger, INamespacePolicyResolver policyResolver, IReferencesStore referencesStore, IPeerStatusService peerStatusService) : base(serviceName: nameof(BlobStoreConsistencyCheckService), TimeSpan.FromSeconds(settings.CurrentValue.ConsistencyCheckPollFrequencySeconds), new ConsistencyState(), logger)
+		public RefStoreConsistencyCheckService(IOptionsMonitor<ConsistencyCheckSettings> settings, IOptionsMonitor<UnrealCloudDDCSettings> unrealCloudDDCSettings, IServiceProvider provider, ILeaderElection leaderElection, IRefService refService, IBlobIndex blobIndex, Tracer tracer, ILogger<BlobStoreConsistencyCheckService> logger, INamespacePolicyResolver policyResolver, IReferencesStore referencesStore, IPeerStatusService peerStatusService) : base(serviceName: nameof(RefStoreConsistencyCheckService), TimeSpan.FromSeconds(settings.CurrentValue.ConsistencyCheckPollFrequencySeconds), new ConsistencyState(), logger)
 		{
 			_settings = settings;
 			_unrealCloudDDCSettings = unrealCloudDDCSettings;
@@ -121,6 +121,16 @@ namespace Jupiter.Implementation
 
 		private async Task<bool> VerifyRegionalConsistencyAsync(NamespaceId ns, BucketId bucket, RefId refId)
 		{
+			string[] fields = new string[] { "namespace", "bucket", "ref", "last_access" };
+			(RefRecord refRecord, BlobContents? _) = await _refService.GetAsync(ns, bucket, refId, fields: fields, doLastAccessTracking: false);
+			// last access in this case is going to be the creation date
+			// TODO: Not sure if we really need to exclude new refs like this. Attempting without this for now.
+			/*if (refRecord.LastAccess > DateTime.Now.AddHours(-2))
+			{
+				// if the ref was created in the last two hours we ignore it as replication may not have had time to run yet
+				return false;
+			}*/
+
 			List<BlobId> blobs = await _refService.GetReferencedBlobsAsync(ns, bucket, refId, ignoreMissingBlobs: true);
 			Dictionary<string, Dictionary<string, bool>> blobStatePerRegion = new Dictionary<string, Dictionary<string, bool>>();
 
@@ -169,7 +179,7 @@ namespace Jupiter.Implementation
 					countOfMissingBlobs++;
 				}
 			}
-			_logger.LogWarning("Regional inconsistency for ref {Namespace} {Bucket} {RefId} missing a total of {CountOfBlobs} blob(s) in these regions: {Regions} ", ns, bucket, refId, countOfMissingBlobs, regionsWithMissingBlobs);
+			_logger.LogWarning("Regional inconsistency for ref {Namespace} {Bucket} {RefId} missing a total of {CountOfBlobs} blob(s) in these regions: {Regions} . Ref was created at {CreationTime}", ns, bucket, refId, countOfMissingBlobs, regionsWithMissingBlobs, refRecord.LastAccess);
 
 			return blobMissing;
 		}
