Cloud DDC - Fixed log context used for RefStoreConsistencyChecks

diff --git a/Engine/Source/Programs/UnrealCloudDDC/Jupiter/Implementation/Blob/RefStoreConsistencyCheckService.cs b/Engine/Source/Programs/UnrealCloudDDC/Jupiter/Implementation/Blob/RefStoreConsistencyCheckService.cs
--- a/Engine/Source/Programs/UnrealCloudDDC/Jupiter/Implementation/Blob/RefStoreConsistencyCheckService.cs
+++ b/Engine/Source/Programs/UnrealCloudDDC/Jupiter/Implementation/Blob/RefStoreConsistencyCheckService.cs
@@ -38,7 +38,7 @@ namespace Jupiter.Implementation
 			return _settings.CurrentValue.EnableRefStoreChecks;
 		}
 
-		public RefStoreConsistencyCheckService(IOptionsMonitor<ConsistencyCheckSettings> settings, IOptionsMonitor<UnrealCloudDDCSettings> unrealCloudDDCSettings, IServiceProvider provider, ILeaderElection leaderElection, IRefService refService, IBlobIndex blobIndex, Tracer tracer, ILogger<BlobStoreConsistencyCheckService> logger, INamespacePolicyResolver policyResolver, IReferencesStore referencesStore, IPeerStatusService peerStatusService) : base(serviceName: nameof(RefStoreConsistencyCheckService), TimeSpan.FromSeconds(settings.CurrentValue.ConsistencyCheckPollFrequencySeconds), new ConsistencyState(), logger)
+		public RefStoreConsistencyCheckService(IOptionsMonitor<ConsistencyCheckSettings> settings, IOptionsMonitor<UnrealCloudDDCSettings> unrealCloudDDCSettings, IServiceProvider provider, ILeaderElection leaderElection, IRefService refService, IBlobIndex blobIndex, Tracer tracer, ILogger<RefStoreConsistencyCheckService> logger, INamespacePolicyResolver policyResolver, IReferencesStore referencesStore, IPeerStatusService peerStatusService) : base(serviceName: nameof(RefStoreConsistencyCheckService), TimeSpan.FromSeconds(settings.CurrentValue.ConsistencyCheckPollFrequencySeconds), new ConsistencyState(), logger)
 		{
 			_settings = settings;
 			_unrealCloudDDCSettings = unrealCloudDDCSettings;
