Cloud DDC - Upgraded Serilog.Enrichers.ClientInfo to fix vulnerability warning

diff --git a/Engine/Source/Programs/UnrealCloudDDC/Jupiter.Common/Jupiter.Common.csproj b/Engine/Source/Programs/UnrealCloudDDC/Jupiter.Common/Jupiter.Common.csproj
--- a/Engine/Source/Programs/UnrealCloudDDC/Jupiter.Common/Jupiter.Common.csproj
+++ b/Engine/Source/Programs/UnrealCloudDDC/Jupiter.Common/Jupiter.Common.csproj
@@ -56,7 +56,7 @@
 
     <!--Support reading yaml config files -->
     <PackageReference Include="NetEscapades.Configuration.Yaml" Version="3.1.0" />
-    <PackageReference Include="Serilog.Enrichers.ClientInfo" Version="2.0.3" />
+    <PackageReference Include="Serilog.Enrichers.ClientInfo" Version="2.1.1" />
     <PackageReference Include="Serilog.Extensions.Hosting" Version="8.0.0" />
     <PackageReference Include="SharpYaml" Version="2.1.0" />
 
diff --git a/Engine/Source/Programs/UnrealCloudDDC/Jupiter/Jupiter.csproj b/Engine/Source/Programs/UnrealCloudDDC/Jupiter/Jupiter.csproj
--- a/Engine/Source/Programs/UnrealCloudDDC/Jupiter/Jupiter.csproj
+++ b/Engine/Source/Programs/UnrealCloudDDC/Jupiter/Jupiter.csproj
@@ -80,7 +80,7 @@
     <PackageReference Include="Microsoft.Extensions.Logging.Debug" Version="8.0.0" />
     <PackageReference Include="MongoDB.Driver" Version="2.23.1" />
     <PackageReference Include="Newtonsoft.Json" Version="13.0.3" />
-    <PackageReference Include="Serilog.Enrichers.ClientInfo" Version="2.0.3" />
+    <PackageReference Include="Serilog.Enrichers.ClientInfo" Version="2.1.1" />
     <PackageReference Include="SerilogAnalyzer" Version="0.15.0" />
     <PackageReference Include="Crc32.NET" Version="1.2.0" />
     <PackageReference Include="K4os.Compression.LZ4" Version="1.3.6" />
