Cloud DDC - Fixed issue were securityContexts for the pods was unable to be specified.

diff --git a/Engine/Source/Programs/UnrealCloudDDC/Helm/UnrealCloudDDC/templates/deployment.yaml b/Engine/Source/Programs/UnrealCloudDDC/Helm/UnrealCloudDDC/templates/deployment.yaml
--- a/Engine/Source/Programs/UnrealCloudDDC/Helm/UnrealCloudDDC/templates/deployment.yaml
+++ b/Engine/Source/Programs/UnrealCloudDDC/Helm/UnrealCloudDDC/templates/deployment.yaml
@@ -41,6 +41,10 @@ spec:
         podAntiAffinity: {{- include "common.affinities.pods" (dict "type" .Values.podAntiAffinityPreset "context" $) | nindent 10 }}
         nodeAffinity: {{- include "common.affinities.nodes" (dict "type" .Values.nodeAffinityPreset.type "key" .Values.nodeAffinityPreset.key "values" .Values.nodeAffinityPreset.values) | nindent 10 }}
       {{- end }}
+      {{- if .Values.securityContext.enabled }}
+      {{ $securityContext := unset .Values.securityContext "enabled" }}
+      securityContext: {{- include "common.tplvalues.render" (dict "value" $securityContext "context" $) | nindent 8 }}
+      {{- end }}
       {{- if .Values.nodeSelector }}
       nodeSelector: {{- include "common.tplvalues.render" (dict "value" .Values.nodeSelector "context" $) | nindent 8 }}
       {{- end }}
@@ -88,9 +92,6 @@ spec:
         {{- if .Values.env }}
         {{- include "common.tplvalues.render" ( dict "value" .Values.env "context" $ ) | nindent 8 }}
         {{- end }}
-        {{- if .Values.securityContext.enabled }}
-        securityContext: {{- include "common.tplvalues.render" (dict "value" .Values.securityContext "context" $) | nindent 8 }}
-        {{- end }}
         ports:
         - name: http
           containerPort: 80
