Cloud DDC - Surpressing exceptions raised during ref store consistency check - we do not particularly care if a ref is in a partial state for this operation.

diff --git a/Engine/Source/Programs/UnrealCloudDDC/Jupiter/Implementation/Blob/RefStoreConsistencyCheckService.cs b/Engine/Source/Programs/UnrealCloudDDC/Jupiter/Implementation/Blob/RefStoreConsistencyCheckService.cs
--- a/Engine/Source/Programs/UnrealCloudDDC/Jupiter/Implementation/Blob/RefStoreConsistencyCheckService.cs
+++ b/Engine/Source/Programs/UnrealCloudDDC/Jupiter/Implementation/Blob/RefStoreConsistencyCheckService.cs
@@ -98,7 +98,16 @@ namespace Jupiter.Implementation
 				{
 					if (_settings.CurrentValue.RegionalConsistencyCheckNamespaces.Contains(ns.ToString()))
 					{
-						bool missing = await VerifyRegionalConsistencyAsync(ns, bucket, refId);
+						bool missing = false;
+						try
+						{
+							missing = await VerifyRegionalConsistencyAsync(ns, bucket, refId);
+						}
+						catch (Exception e)
+						{
+							_logger.LogWarning(e, "Unknown exception with message {Message} when checking ref store consistency. Ignoring.", e.Message);
+						}
+
 						if (missing)
 						{
 							Interlocked.Increment(ref countOfRegionalInconsistentRefs);
