Cloud DDC - Handle bucket already exists errors from S3 while creating the bucket (this indicates that we had two attempts at creating the bucket at the same time, so the second one fails, as we just care about the bucket existing we can safely ignore this error)

diff --git a/Engine/Source/Programs/UnrealCloudDDC/Jupiter/Implementation/Blob/AmazonS3Store.cs b/Engine/Source/Programs/UnrealCloudDDC/Jupiter/Implementation/Blob/AmazonS3Store.cs
--- a/Engine/Source/Programs/UnrealCloudDDC/Jupiter/Implementation/Blob/AmazonS3Store.cs
+++ b/Engine/Source/Programs/UnrealCloudDDC/Jupiter/Implementation/Blob/AmazonS3Store.cs
@@ -248,7 +248,19 @@ namespace Jupiter.Implementation
 							UseClientRegion = true
 						};
 
-						await _amazonS3.PutBucketAsync(putBucketRequest, cancellationToken);
+						try
+						{
+							await _amazonS3.PutBucketAsync(putBucketRequest, cancellationToken);
+						}
+						catch (AmazonS3Exception e)
+						{
+							if (e.StatusCode == HttpStatusCode.Conflict)
+							{
+								// bucket already exists so no need to try and create it, most likely a race condition with another write operation that already tried to create the bucket at the same time in this or another instance
+								// as this will end up with the bucket existing we can just ignore the error
+							}
+							throw;
+						}
 					}
 					_bucketExistenceChecked = true;
 				}
