Cloud DDC - Fixed for tests failing when s3 bucket already exists

diff --git a/Engine/Source/Programs/UnrealCloudDDC/Jupiter/Implementation/Blob/AmazonS3Store.cs b/Engine/Source/Programs/UnrealCloudDDC/Jupiter/Implementation/Blob/AmazonS3Store.cs
--- a/Engine/Source/Programs/UnrealCloudDDC/Jupiter/Implementation/Blob/AmazonS3Store.cs
+++ b/Engine/Source/Programs/UnrealCloudDDC/Jupiter/Implementation/Blob/AmazonS3Store.cs
@@ -259,7 +259,10 @@ namespace Jupiter.Implementation
 								// bucket already exists so no need to try and create it, most likely a race condition with another write operation that already tried to create the bucket at the same time in this or another instance
 								// as this will end up with the bucket existing we can just ignore the error
 							}
-							throw;
+							else
+							{
+								throw;
+							}
 						}
 					}
 					_bucketExistenceChecked = true;
