Cloud DDC - Removed internal filter from admin controller, meaning that these apis can be accessed on the public endpoints (but requires a admin token). They are not super useful if they can not be easily accessed as they are usually used by humans and not by the system itself (which is were the internal filter is mostly useful).

diff --git a/Engine/Source/Programs/UnrealCloudDDC/Jupiter/Controllers/AdminController.cs b/Engine/Source/Programs/UnrealCloudDDC/Jupiter/Controllers/AdminController.cs
--- a/Engine/Source/Programs/UnrealCloudDDC/Jupiter/Controllers/AdminController.cs
+++ b/Engine/Source/Programs/UnrealCloudDDC/Jupiter/Controllers/AdminController.cs
@@ -15,7 +15,6 @@ namespace Jupiter.Controllers
 {
 	[ApiController]
 	[Route("api/v1/admin")]
-	[InternalApiFilter]
 	[Authorize]
 	public class AdminController : Controller
 	{
@@ -69,7 +68,7 @@ namespace Jupiter.Controllers
 		{
 			ActionResult? result = await _requestHelper.HasAccessForGlobalOperationsAsync(User, new[] { JupiterAclAction.AdminAction });
 			if (result != null)
-			{
+			{	
 				return result;
 			}
 
