Cloud DDC - Fixed issue were settings applied to the Jupiter section were not being passed if you were also using nginx domain sockets. This esspecially impacted the 'PendingConnectionMax' setting which you bump when you also bump nginx values.

diff --git a/Engine/Source/Programs/UnrealCloudDDC/Helm/UnrealCloudDDC/templates/configmap.yaml b/Engine/Source/Programs/UnrealCloudDDC/Helm/UnrealCloudDDC/templates/configmap.yaml
--- a/Engine/Source/Programs/UnrealCloudDDC/Helm/UnrealCloudDDC/templates/configmap.yaml
+++ b/Engine/Source/Programs/UnrealCloudDDC/Helm/UnrealCloudDDC/templates/configmap.yaml
@@ -71,9 +71,14 @@ UnrealCloudDDC:
   {{- end -}}
 {{- end }}
 
-{{- if $config.Jupiter }}
+{{- if or ($config.Jupiter) (.Values.nginx.useDomainSockets) }}
 Jupiter: 
 {{ $config.Jupiter | toYaml | nindent 2 }}
+{{ if .Values.nginx.useDomainSockets }}
+  UseDomainSockets: true
+  DomainSocketsRoot: "/nginx"
+  ChmodDomainSockets: true
+{{- end }}
 {{- end }}
 
 {{- if $config.Debug }}
@@ -81,13 +86,6 @@ Debug:
 {{ $config.Debug | toYaml | nindent 2 }}
 {{- end }}
 
-{{ if .Values.nginx.useDomainSockets }}
-Jupiter:
-  UseDomainSockets: true
-  DomainSocketsRoot: "/nginx"
-  ChmodDomainSockets: true
-{{ end }}
-
 {{ if .Values.nginx.enabled }}
 Nginx:
   UseNginxRedirect: true
