Cloud DDC - Fixed issue were mongo content id comparisions were failing because .Equals is not supported in mongo queries for strings.

diff --git a/Engine/Source/Programs/UnrealCloudDDC/Jupiter/Implementation/ContentId/MongoContentIdStore.cs b/Engine/Source/Programs/UnrealCloudDDC/Jupiter/Implementation/ContentId/MongoContentIdStore.cs
--- a/Engine/Source/Programs/UnrealCloudDDC/Jupiter/Implementation/ContentId/MongoContentIdStore.cs
+++ b/Engine/Source/Programs/UnrealCloudDDC/Jupiter/Implementation/ContentId/MongoContentIdStore.cs
@@ -1,6 +1,5 @@
 // Copyright Epic Games, Inc. All Rights Reserved.
 
-using System;
 using System.Collections.Generic;
 using System.Linq;
 using System.Threading;
@@ -42,7 +41,7 @@ namespace Jupiter.Implementation
 			IMongoCollection<MongoContentIdModelV0> collection = GetCollection<MongoContentIdModelV0>();
 
 			IAsyncCursor<MongoContentIdModelV0> cursor =
-				await collection.FindAsync(model => model.Ns.Equals(ns.ToString(), StringComparison.OrdinalIgnoreCase) && model.ContentId.Equals(contentId.ToString(), StringComparison.Ordinal), cancellationToken: cancellationToken);
+				await collection.FindAsync(model => model.Ns == ns.ToString() && model.ContentId == contentId.ToString(), cancellationToken: cancellationToken);
 
 			MongoContentIdModelV0 model = await cursor.FirstOrDefaultAsync(cancellationToken);
 			if (model != null)
