Cloud DDC - Improved logging message for filesystem cleanup when run against a specific namespace.

diff --git a/Engine/Source/Programs/UnrealCloudDDC/Jupiter/Implementation/Blob/FileSystemStore.cs b/Engine/Source/Programs/UnrealCloudDDC/Jupiter/Implementation/Blob/FileSystemStore.cs
--- a/Engine/Source/Programs/UnrealCloudDDC/Jupiter/Implementation/Blob/FileSystemStore.cs
+++ b/Engine/Source/Programs/UnrealCloudDDC/Jupiter/Implementation/Blob/FileSystemStore.cs
@@ -214,23 +214,23 @@ namespace Jupiter.Implementation
 			for (int i = 0; i < 5; i++)
 			{
 				long size = await CalculateDiskSpaceUsedAsync(ns);
-
+				string namespaceMessage = ns != null ? $" Namespace: '{ns.ToString()}' .": string.Empty;
 				// first check to see if we should trigger at all, this happens for each run but only really matters for the first attempt
 				if (size < triggerSize)
 				{
-					_logger.LogInformation("Filesystem cleanup not running. Disksize used: {UsedDiskSize} . Trigger size was {TriggerSize}", size, triggerSize);
+					_logger.LogInformation("Filesystem cleanup not running. Disksize used: '{UsedDiskSize}'.{NamespaceMessage} Trigger size was {TriggerSize}", size, namespaceMessage, triggerSize);
 					return countOfBlobsRemoved;
 				}
 
 				// then check if we have reached the target size, if not we should continue running
 				if (size <= targetSize)
 				{
-					_logger.LogInformation("Filesystem cleanup reached target size. Disksize used: {UsedDiskSize} . Target size was {TargetSize}", size, targetSize);
+					_logger.LogInformation("Filesystem cleanup reached target size. Disksize used: '{UsedDiskSize}'.{NamespaceMessage} Target size was {TargetSize}", size, namespaceMessage, targetSize);
 
 					return countOfBlobsRemoved;
 				}
 
-				_logger.LogInformation("Filesystem cleanup running. Disksize used: {UsedDiskSize} . Trigger size was {TriggerSize}", size, triggerSize);
+				_logger.LogInformation("Filesystem cleanup running. Disksize used: '{UsedDiskSize}'.{NamespaceMessage} Trigger size was {TriggerSize}", size, namespaceMessage, triggerSize);
 
 				// define progressively shorter windows of how long we keep data around for, based on their last write time
 				DateTime[] cutoffPeriods = new DateTime[]
