Cloud DDC - Upgraded Blake3 version to fix inconsistency with EpicGames.Iohash

diff --git a/Engine/Source/Programs/UnrealCloudDDC/Jupiter.Common/Jupiter.Common.csproj b/Engine/Source/Programs/UnrealCloudDDC/Jupiter.Common/Jupiter.Common.csproj
--- a/Engine/Source/Programs/UnrealCloudDDC/Jupiter.Common/Jupiter.Common.csproj
+++ b/Engine/Source/Programs/UnrealCloudDDC/Jupiter.Common/Jupiter.Common.csproj
@@ -23,7 +23,7 @@
   </PropertyGroup>
 
   <ItemGroup>
-    <PackageReference Include="Blake3" Version="1.0.0" />
+    <PackageReference Include="Blake3" Version="1.1.0" />
 
     <PackageReference Include="AspNetCore.HealthChecks.Publisher.Datadog" Version="8.0.0" />
     <PackageReference Include="AWSSDK.Extensions.NETCore.Setup" Version="3.7.300" />
