Cloud DDC - No longer using if not exists when preventing overwrites as this results in scylla ignoring our consistency level and running the operation as a serial consistency (instead of local one) meaning the operation goes from 10ms to 500ms (because it needs to wait for all nodes in the cluster).

Now we read before writing but with the normal local_one consistency, would be bad if we mutate this value frequently but we do not so this should be safe and quicker.

diff --git a/Engine/Source/Programs/UnrealCloudDDC/Jupiter/Implementation/Objects/ScyllaReferencesStore.cs b/Engine/Source/Programs/UnrealCloudDDC/Jupiter/Implementation/Objects/ScyllaReferencesStore.cs
--- a/Engine/Source/Programs/UnrealCloudDDC/Jupiter/Implementation/Objects/ScyllaReferencesStore.cs
+++ b/Engine/Source/Programs/UnrealCloudDDC/Jupiter/Implementation/Objects/ScyllaReferencesStore.cs
@@ -161,8 +161,8 @@ namespace Jupiter.Implementation
 			bool allowOverwrites = policy.AllowOverwritesOfRefs || allowOverwrite;
 			if (!allowOverwrites)
 			{
-				AppliedInfo<ScyllaObject>? applied = await _mapper.InsertIfNotExistsAsync<ScyllaObject>(new ScyllaObject(ns, bucket, name, blob, blobHash, isFinalized), ttl: ttl, insertNulls: false);
-				if (!applied.Applied)
+				bool insertRef = false;
+				try
 				{
 					RefRecord oldRecord = await GetAsync(ns, bucket, name, IReferencesStore.FieldFlags.None, IReferencesStore.OperationFlags.None, cancellationToken);
 
@@ -171,8 +171,17 @@ namespace Jupiter.Implementation
 						// blob was not the same, e.g. we attempted to change the value, this is not allowed
 						throw new RefAlreadyExistsException(ns, bucket, name);
 					}
+					
+					// new record is the same as the old, this is fine there is no need to write anything
+				}
+				catch (RefNotFoundException)
+				{
+					insertRef = true;
+				}
 
-					// new record is the same as the old, this is fine we can just carry on
+				if (insertRef)
+				{
+					await _mapper.InsertAsync<ScyllaObject>(new ScyllaObject(ns, bucket, name, blob, blobHash, isFinalized), ttl: ttl, insertNulls: false);
 				}
 			}
 			else
