Cloud DDC - Joining the replication task every 15 minutes to make sure it gets some cpu time, hopefully resolves issues with tasks sometimes stopping until the worker is restarted in conditions with a lot of parallel tasks.

diff --git a/Engine/Source/Programs/UnrealCloudDDC/Jupiter/Implementation/Replication/ReplicationService.cs b/Engine/Source/Programs/UnrealCloudDDC/Jupiter/Implementation/Replication/ReplicationService.cs
--- a/Engine/Source/Programs/UnrealCloudDDC/Jupiter/Implementation/Replication/ReplicationService.cs
+++ b/Engine/Source/Programs/UnrealCloudDDC/Jupiter/Implementation/Replication/ReplicationService.cs
@@ -23,7 +23,7 @@ namespace Jupiter.Implementation
 		private readonly IOptionsMonitor<ReplicationSettings> _settings;
 		private readonly ILeaderElection _leaderElection;
 		private readonly ILogger _logger;
-		private readonly Dictionary<string, Task<bool>> _currentReplications = new Dictionary<string, Task<bool>>();
+		private readonly Dictionary<string, PendingReplication> _currentReplications = new Dictionary<string, PendingReplication>();
 
 		protected override bool ShouldStartPolling()
 		{
@@ -96,8 +96,9 @@ namespace Jupiter.Implementation
 
 			foreach (IReplicator replicator in state.Replicators)
 			{
-				if (_currentReplications.TryGetValue(replicator.Info.ReplicatorName, out Task<bool>? replicationTask))
+				if (_currentReplications.TryGetValue(replicator.Info.ReplicatorName, out PendingReplication? pendingReplication))
 				{
+					Task replicationTask = pendingReplication.Task;
 					if (replicationTask.IsCompleted)
 					{
 						_currentReplications.Remove(replicator.Info.ReplicatorName);
@@ -116,16 +117,25 @@ namespace Jupiter.Implementation
 					}
 					else
 					{
-						_logger.LogDebug("Replication of replicator: {Name} is still running", replicator.Info.ReplicatorName);
 						// if the replication is still running let it continue to run
+						_logger.LogDebug("Replication of replicator: {Name} is still running", replicator.Info.ReplicatorName);
+
+						// if the replication has been running for more then 15 minutes join the task for a little while to make sure it gets some cpu time
+						if (DateTime.Now - pendingReplication.LastForceRun > TimeSpan.FromMinutes(15))
+						{
+							_logger.LogDebug("Replication of replicator: {Name} has been running for a long time. Joining for a while to make sure it has time.", replicator.Info.ReplicatorName);
+
+							await Task.WhenAny(replicationTask, Task.Delay(TimeSpan.FromSeconds(30), cancellationToken));
+							pendingReplication.LastForceRun = DateTime.Now;
+						}
 						continue;
 					}
 				}
 
 				// start a new run of the replication
 				_logger.LogDebug("Triggering new replication of replicator: {Name}", replicator.Info.ReplicatorName);
-				Task<bool> newReplication = replicator.TriggerNewReplicationsAsync();
-				_currentReplications[replicator.Info.ReplicatorName] = newReplication;
+
+				_currentReplications[replicator.Info.ReplicatorName] = new PendingReplication { Task =  replicator.TriggerNewReplicationsAsync(), LastForceRun = DateTime.Now};
 			}
 
 			if (state.Replicators.Count == 0)
@@ -137,6 +147,12 @@ namespace Jupiter.Implementation
 
 		}
 
+		private class PendingReplication
+		{
+			public Task<bool> Task { get; init; } = null!;
+			public DateTime LastForceRun { get; set; }
+		}
+
 		protected override async Task OnStopping(ReplicationState state)
 		{
 			await Task.WhenAll(state.Replicators.Select(replicator => replicator.StopReplicatingAsync()).ToArray());
