Cloud DDC - Upgraded Microsoft.Extensions.Caching.Memory to 8.0.1 to fix build issues after change 37033668

diff --git a/Engine/Source/Programs/UnrealCloudDDC/Jupiter/Jupiter.csproj b/Engine/Source/Programs/UnrealCloudDDC/Jupiter/Jupiter.csproj
--- a/Engine/Source/Programs/UnrealCloudDDC/Jupiter/Jupiter.csproj
+++ b/Engine/Source/Programs/UnrealCloudDDC/Jupiter/Jupiter.csproj
@@ -76,7 +76,7 @@
     <PackageReference Include="Azure.Security.KeyVault.Secrets" Version="4.5.0" />
     <PackageReference Include="Azure.Storage.Blobs" Version="12.15.1" />
     <PackageReference Include="KubernetesClient" Version="13.0.1" />
-    <PackageReference Include="Microsoft.Extensions.Caching.Memory" Version="8.0.0" />
+    <PackageReference Include="Microsoft.Extensions.Caching.Memory" Version="8.0.1" />
     <PackageReference Include="Microsoft.Extensions.Logging.Debug" Version="8.0.0" />
     <PackageReference Include="MongoDB.Driver" Version="2.23.1" />
     <PackageReference Include="Newtonsoft.Json" Version="13.0.3" />
