Cloud DDC - Fixed issue where the peer store would call into itself to see if the blob was present which makes no sense.

diff --git a/Engine/Source/Programs/UnrealCloudDDC/Helm/UnrealCloudDDC/templates/deployment.yaml b/Engine/Source/Programs/UnrealCloudDDC/Helm/UnrealCloudDDC/templates/deployment.yaml
--- a/Engine/Source/Programs/UnrealCloudDDC/Helm/UnrealCloudDDC/templates/deployment.yaml
+++ b/Engine/Source/Programs/UnrealCloudDDC/Helm/UnrealCloudDDC/templates/deployment.yaml
@@ -89,6 +89,14 @@ spec:
           valueFrom:
             fieldRef:
               fieldPath: metadata.namespace
+        - name: Kubernetes__PodName
+          valueFrom:
+            fieldRef:
+              fieldPath: metadata.name
+        - name: Kubernetes__PodIp
+          valueFrom:
+            fieldRef:
+              fieldPath: metadata.podIP
         {{- if .Values.env }}
         {{- include "common.tplvalues.render" ( dict "value" .Values.env "context" $ ) | nindent 8 }}
         {{- end }}
diff --git a/Engine/Source/Programs/UnrealCloudDDC/Jupiter/Implementation/Blob/PeerBlobStore.cs b/Engine/Source/Programs/UnrealCloudDDC/Jupiter/Implementation/Blob/PeerBlobStore.cs
--- a/Engine/Source/Programs/UnrealCloudDDC/Jupiter/Implementation/Blob/PeerBlobStore.cs
+++ b/Engine/Source/Programs/UnrealCloudDDC/Jupiter/Implementation/Blob/PeerBlobStore.cs
@@ -292,6 +292,7 @@ namespace Jupiter.Implementation
 
 		private async IAsyncEnumerable<string> EnumeratePodsAsync()
 		{
+			string currentPodName = _leaderSettings.CurrentValue.PodName;
 			V1PodList podList = await _client.ListNamespacedPodAsync(_leaderSettings.CurrentValue.Namespace, labelSelector: _leaderSettings.CurrentValue.PeerPodLabelSelector);
 
 			foreach (V1Pod pod in podList.Items)
@@ -300,6 +301,12 @@ namespace Jupiter.Implementation
 				{
 					continue;
 				}
+
+				// skip the current pod
+				if (pod.Name() == currentPodName)
+				{
+					continue;
+				}
 				string ip = pod.Status.PodIP;
 				if (string.IsNullOrEmpty(ip))
 				{
diff --git a/Engine/Source/Programs/UnrealCloudDDC/Jupiter/Implementation/LeaderElection/KubernetesLeaderElection.cs b/Engine/Source/Programs/UnrealCloudDDC/Jupiter/Implementation/LeaderElection/KubernetesLeaderElection.cs
--- a/Engine/Source/Programs/UnrealCloudDDC/Jupiter/Implementation/LeaderElection/KubernetesLeaderElection.cs
+++ b/Engine/Source/Programs/UnrealCloudDDC/Jupiter/Implementation/LeaderElection/KubernetesLeaderElection.cs
@@ -25,6 +25,8 @@ namespace Jupiter.Implementation.LeaderElection
 
 		[Required]
 		public string Namespace { get; set; } = "";
+		[Required] public string PodName { get; set; } = "";
+		[Required] public string PodIp { get; set; } = "";
 
 		public string PeerPodLabelSelector { get; set; } = "app.kubernetes.io/name=unreal-cloud-ddc";
 	}
