Cloud DDC - Fixed issue were the speculative blob replication wouldn't work when combined with on demand replication (as it noticed that blobs were always available in other regions and thus already existed)

diff --git a/Engine/Source/Programs/UnrealCloudDDC/Jupiter/Implementation/Blob/BlobService.cs b/Engine/Source/Programs/UnrealCloudDDC/Jupiter/Implementation/Blob/BlobService.cs
--- a/Engine/Source/Programs/UnrealCloudDDC/Jupiter/Implementation/Blob/BlobService.cs
+++ b/Engine/Source/Programs/UnrealCloudDDC/Jupiter/Implementation/Blob/BlobService.cs
@@ -666,7 +666,7 @@ public class BlobService : IBlobService
 		return request;
 	}
 
-	public async Task<bool> ExistsAsync(NamespaceId ns, BlobId blob, List<string>? storageLayers = null, CancellationToken cancellationToken = default)
+	public async Task<bool> ExistsAsync(NamespaceId ns, BlobId blob, List<string>? storageLayers = null, bool ignoreRemoteBlobs = false, CancellationToken cancellationToken = default)
 	{
 		bool exists = await ExistsInStoresAsync(ns, blob, storageLayers, cancellationToken);
 		if (exists)
@@ -680,7 +680,7 @@ public class BlobService : IBlobService
 			return await ExistsInStoresAsync(policy.FallbackNamespace.Value, blob, storageLayers, cancellationToken);
 		}
 
-		if (ShouldFetchBlobOnDemand(ns))
+		if (ShouldFetchBlobOnDemand(ns) && !ignoreRemoteBlobs)
 		{
 			return await ExistsInRemoteAsync(ns, blob, cancellationToken);
 		}
diff --git a/Engine/Source/Programs/UnrealCloudDDC/Jupiter/Implementation/Blob/IBlobService.cs b/Engine/Source/Programs/UnrealCloudDDC/Jupiter/Implementation/Blob/IBlobService.cs
--- a/Engine/Source/Programs/UnrealCloudDDC/Jupiter/Implementation/Blob/IBlobService.cs
+++ b/Engine/Source/Programs/UnrealCloudDDC/Jupiter/Implementation/Blob/IBlobService.cs
@@ -30,7 +30,7 @@ public interface IBlobService
 
 	Task<BlobContents> ReplicateObjectAsync(NamespaceId ns, BlobId blob, bool force = false, BucketId? bucketHint = null, CancellationToken cancellationToken = default);
 
-	Task<bool> ExistsAsync(NamespaceId ns, BlobId blob, List<string>? storageLayers = null, CancellationToken cancellationToken = default);
+	Task<bool> ExistsAsync(NamespaceId ns, BlobId blob, List<string>? storageLayers = null, bool ignoreRemoteBlobs = false, CancellationToken cancellationToken = default);
 
 	/// <summary>
 	/// Checks that the blob exists in the root store, the store which is last in the list and thus is intended to have every blob in it
diff --git a/Engine/Source/Programs/UnrealCloudDDC/Jupiter/Implementation/Replication/BlobsReplicator.cs b/Engine/Source/Programs/UnrealCloudDDC/Jupiter/Implementation/Replication/BlobsReplicator.cs
--- a/Engine/Source/Programs/UnrealCloudDDC/Jupiter/Implementation/Replication/BlobsReplicator.cs
+++ b/Engine/Source/Programs/UnrealCloudDDC/Jupiter/Implementation/Replication/BlobsReplicator.cs
@@ -439,7 +439,8 @@ namespace Jupiter.Implementation
 				.SetAttribute("operation.name", "replicator.replicate_blob")
 				.SetAttribute("resource.name", $"{ns}.{blob}");
 
-			bool exists = await _blobService.ExistsAsync(ns, blob, storageLayers: null, cancellationToken);
+			// the point of the replicator is to transfer the blobs locally, thus we ignore if they exist in remote locations
+			bool exists = await _blobService.ExistsAsync(ns, blob, storageLayers: null, ignoreRemoteBlobs: true, cancellationToken: cancellationToken);
 			if (exists)
 			{
 				_logger.LogDebug("Not replicating blob {Blob} in {Namespace} as it already existed.", blob, ns);
