Cloud DDC - Removed pod ip setting as its not always supported in k8s and we do not need it for anything right now.

diff --git a/Engine/Source/Programs/UnrealCloudDDC/Helm/UnrealCloudDDC/templates/deployment.yaml b/Engine/Source/Programs/UnrealCloudDDC/Helm/UnrealCloudDDC/templates/deployment.yaml
--- a/Engine/Source/Programs/UnrealCloudDDC/Helm/UnrealCloudDDC/templates/deployment.yaml
+++ b/Engine/Source/Programs/UnrealCloudDDC/Helm/UnrealCloudDDC/templates/deployment.yaml
@@ -93,10 +93,6 @@ spec:
           valueFrom:
             fieldRef:
               fieldPath: metadata.name
-        - name: Kubernetes__PodIp
-          valueFrom:
-            fieldRef:
-              fieldPath: metadata.podIP
         {{- if .Values.env }}
         {{- include "common.tplvalues.render" ( dict "value" .Values.env "context" $ ) | nindent 8 }}
         {{- end }}
diff --git a/Engine/Source/Programs/UnrealCloudDDC/Jupiter/Implementation/LeaderElection/KubernetesLeaderElection.cs b/Engine/Source/Programs/UnrealCloudDDC/Jupiter/Implementation/LeaderElection/KubernetesLeaderElection.cs
--- a/Engine/Source/Programs/UnrealCloudDDC/Jupiter/Implementation/LeaderElection/KubernetesLeaderElection.cs
+++ b/Engine/Source/Programs/UnrealCloudDDC/Jupiter/Implementation/LeaderElection/KubernetesLeaderElection.cs
@@ -26,7 +26,6 @@ namespace Jupiter.Implementation.LeaderElection
 		[Required]
 		public string Namespace { get; set; } = "";
 		[Required] public string PodName { get; set; } = "";
-		[Required] public string PodIp { get; set; } = "";
 
 		public string PeerPodLabelSelector { get; set; } = "app.kubernetes.io/name=unreal-cloud-ddc";
 	}
