Cloud DDC - Added tracing scope when reading blob index from scylla.

diff --git a/Engine/Source/Programs/UnrealCloudDDC/Jupiter/Implementation/Blob/ScyllaBlobIndex.cs b/Engine/Source/Programs/UnrealCloudDDC/Jupiter/Implementation/Blob/ScyllaBlobIndex.cs
--- a/Engine/Source/Programs/UnrealCloudDDC/Jupiter/Implementation/Blob/ScyllaBlobIndex.cs
+++ b/Engine/Source/Programs/UnrealCloudDDC/Jupiter/Implementation/Blob/ScyllaBlobIndex.cs
@@ -150,7 +150,12 @@ public class ScyllaBlobIndex : IBlobIndex
 	public async Task<bool> BlobExistsInRegionAsync(NamespaceId ns, BlobId blobIdentifier, string? region = null, CancellationToken cancellationToken = default)
 	{
 		region ??= _jupiterSettings.CurrentValue.CurrentSite;
-		ScyllaBlobIndexEntry? entry = await _mapper.SingleOrDefaultAsync<ScyllaBlobIndexEntry>("WHERE namespace = ? AND blob_id = ? AND region = ?", ns.ToString(), blobIdentifier.HashData, region);
+		ScyllaBlobIndexEntry? entry;
+		{
+			using TelemetrySpan scope = _tracer.BuildScyllaSpan("scylla.get_blob_index").SetAttribute("resource.name", $"{ns}.{blobIdentifier}");
+
+			entry = await _mapper.SingleOrDefaultAsync<ScyllaBlobIndexEntry>("WHERE namespace = ? AND blob_id = ? AND region = ?", ns.ToString(), blobIdentifier.HashData, region);
+		}
 
 		bool blobMissing = entry == null;
 		if (blobMissing && _scyllaSettings.CurrentValue.MigrateFromOldBlobIndex)
