Cloud DDC - Added UE-isBuildMachine header to list of headers we export to tracing

diff --git a/Engine/Source/Programs/UnrealCloudDDC/Jupiter.Common/BaseStartup.cs b/Engine/Source/Programs/UnrealCloudDDC/Jupiter.Common/BaseStartup.cs
--- a/Engine/Source/Programs/UnrealCloudDDC/Jupiter.Common/BaseStartup.cs
+++ b/Engine/Source/Programs/UnrealCloudDDC/Jupiter.Common/BaseStartup.cs
@@ -332,6 +332,14 @@ namespace Jupiter
 								activity.AddTag("ue-request", ueRequestValuesOld.First());
 							}
 						}
+
+						if (request.Headers.TryGetValue("ue-IsBuildMachine", out StringValues ueIsBuildMachine))
+						{
+							if (ueIsBuildMachine.Count != 0)
+							{
+								activity.AddTag("ue-IsBuildMachine", ueIsBuildMachine.First());
+							}
+						}
 					};
 				});
 				builder.SetSampler(new AlwaysOnSampler());
