Cloud DDC - Added tracing information on the amount of operations in a batch request.

diff --git a/Engine/Source/Programs/UnrealCloudDDC/Jupiter/Controllers/ReferencesController.cs b/Engine/Source/Programs/UnrealCloudDDC/Jupiter/Controllers/ReferencesController.cs
--- a/Engine/Source/Programs/UnrealCloudDDC/Jupiter/Controllers/ReferencesController.cs
+++ b/Engine/Source/Programs/UnrealCloudDDC/Jupiter/Controllers/ReferencesController.cs
@@ -1161,6 +1161,8 @@ namespace Jupiter.Controllers
 				return accessResult;
 			}
 
+			Tracer.CurrentSpan.SetAttribute("CountOfOps", ops.Ops.Length);
+
 			HashSet<uint> usedOpIds = new HashSet<uint>();
 			foreach (BatchOps.BatchOp batchOp in ops.Ops)
 			{
