Cloud DDC - Added limit for how wide we go when resolving references, once we are resolving 64 of them for one object we stop and wait for one of the pending resolves to finish. This will be a slowdown for very large objects (oplogs primarily) but it should mean that we do not end up spawning more tasks then we can actually finish.

diff --git a/Engine/Source/Programs/UnrealCloudDDC/Jupiter/Implementation/IReferenceResolver.cs b/Engine/Source/Programs/UnrealCloudDDC/Jupiter/Implementation/IReferenceResolver.cs
--- a/Engine/Source/Programs/UnrealCloudDDC/Jupiter/Implementation/IReferenceResolver.cs
+++ b/Engine/Source/Programs/UnrealCloudDDC/Jupiter/Implementation/IReferenceResolver.cs
@@ -113,13 +113,19 @@ namespace Jupiter.Implementation
 			List<Task<CbObject>> pendingCompactBinaryAttachments = new();
 			List<Task<(ContentId, BlobId[]?)>> pendingContentIdResolves = new();
 
+			const int MaxParallelAttachmentTasks = 64;
+
 			while (pendingCompactBinaryAttachments.Count != 0 || pendingContentIdResolves.Count != 0 || objectsToVisit.Count != 0)
 			{
 				List<Attachment> attachments = new List<Attachment>();
 
 				if (objectsToVisit.TryDequeue(out CbObject? parent))
 				{
-					parent.IterateAttachments(field =>
+					List<CbField> fields = new List<CbField>();
+
+					parent.IterateAttachments(field => fields.Add(field));
+
+					foreach (CbField field in fields)
 					{
 						IoHash attachmentHash = field.AsAttachment();
 
@@ -140,77 +146,36 @@ namespace Jupiter.Implementation
 						{
 							throw new NotImplementedException($"Unknown attachment type for field {field}");
 						}
-					});
-				}
 
-				// check for any content id resolves to finish
-				List<Task<(ContentId, BlobId[]?)>> finishedContentIdResolves = new();
-				foreach (Task<(ContentId, BlobId[]?)> pendingContentIdResolve in pendingContentIdResolves)
-				{
-					if (pendingContentIdResolve.IsCompleted)
-					{
-						ContentId? contentId = null;
-						BlobId[]? resolvedBlobs = null;
-						BlobId? blobIdentifier = null;
-						bool wasContentId = false;
-						try
-						{
-							(contentId, resolvedBlobs) = await pendingContentIdResolve;
-							blobIdentifier = contentId.AsBlobIdentifier();
-							wasContentId = !(resolvedBlobs is { Length: 1 } && resolvedBlobs[0].Equals(blobIdentifier));
-						}
-						catch (InvalidContentIdException)
+						if (pendingContentIdResolves.Count + pendingCompactBinaryAttachments.Count > MaxParallelAttachmentTasks)
 						{
-							resolvedBlobs = null;
+							// wait for one of our lists of pending tasks to be finished to avoid spawning more tasks
+							// also avoiding waiting on a empty list as that is not allowed
+							if (pendingContentIdResolves.Count == 0)
+							{
+								await Task.WhenAny(pendingCompactBinaryAttachments);
+								await ProcessCompactBinaryAttachmentsAsync(pendingCompactBinaryAttachments, objectsToVisit, unresolvedBlobReferences);
+							}
+							else if (pendingCompactBinaryAttachments.Count == 0)
+							{
+								await Task.WhenAny(pendingContentIdResolves);
+								await ProcessContentIdsAsync(attachments, pendingContentIdResolves, unresolvedBlobReferences);
+							}
+							else
+							{
+								await Task.WhenAny(Task.WhenAny(pendingContentIdResolves), Task.WhenAny(pendingCompactBinaryAttachments));
+								
+								await ProcessContentIdsAsync(attachments, pendingContentIdResolves, unresolvedBlobReferences);
+								await ProcessCompactBinaryAttachmentsAsync(pendingCompactBinaryAttachments, objectsToVisit, unresolvedBlobReferences);
+							}
 						}
-						catch (BlobNotFoundException e)
-						{
-							unresolvedBlobReferences.Add(e.Blob);
-						}
-
-						if (wasContentId && resolvedBlobs != null)
-						{
-							attachments.Add(new ContentIdAttachment(contentId!, resolvedBlobs));
-						}
-						else
-						{
-							attachments.Add(new BlobAttachment(blobIdentifier!));
-						}
-
-						finishedContentIdResolves.Add(pendingContentIdResolve);
 					}
 				}
 
-				// cleanup finished tasks
-				foreach (Task<(ContentId, BlobId[]?)> finishedTask in finishedContentIdResolves)
-				{
-					pendingContentIdResolves.Remove(finishedTask);
-				}
+				await ProcessContentIdsAsync(attachments, pendingContentIdResolves, unresolvedBlobReferences);
 
 				// check for any compact binary attachment fetches and add those to the objects we are handling
-				List<Task<CbObject>> finishedCompactBinaryResolves = new();
-				foreach (Task<CbObject> pendingCompactBinaryAttachment in pendingCompactBinaryAttachments)
-				{
-					if (pendingCompactBinaryAttachment.IsCompleted)
-					{
-						try
-						{
-							CbObject childBinaryObject = await pendingCompactBinaryAttachment;
-							objectsToVisit.Enqueue(childBinaryObject);
-						}
-						catch (BlobNotFoundException e)
-						{
-							unresolvedBlobReferences.Add(e.Blob);
-						}
-						finishedCompactBinaryResolves.Add(pendingCompactBinaryAttachment);
-					}
-				}
-
-				// cleanup finished tasks
-				foreach (Task<CbObject> finishedTask in finishedCompactBinaryResolves)
-				{
-					pendingCompactBinaryAttachments.Remove(finishedTask);
-				}
+				await ProcessCompactBinaryAttachmentsAsync(pendingCompactBinaryAttachments, objectsToVisit, unresolvedBlobReferences);
 
 				bool hasWaited = false;
 				// if there are pending resolves left, wait for one of them to finish to avoid busy waiting
@@ -238,6 +203,80 @@ namespace Jupiter.Implementation
 			}
 		}
 
+		private static async Task ProcessCompactBinaryAttachmentsAsync(List<Task<CbObject>> pendingCompactBinaryAttachments, Queue<CbObject> objectsToVisit, List<BlobId> unresolvedBlobReferences)
+		{
+			List<Task<CbObject>> finishedCompactBinaryResolves = new();
+			foreach (Task<CbObject> pendingCompactBinaryAttachment in pendingCompactBinaryAttachments)
+			{
+				if (pendingCompactBinaryAttachment.IsCompleted)
+				{
+					try
+					{
+						CbObject childBinaryObject = await pendingCompactBinaryAttachment;
+						objectsToVisit.Enqueue(childBinaryObject);
+					}
+					catch (BlobNotFoundException e)
+					{
+						unresolvedBlobReferences.Add(e.Blob);
+					}
+					finishedCompactBinaryResolves.Add(pendingCompactBinaryAttachment);
+				}
+			}
+
+			// cleanup finished tasks
+			foreach (Task<CbObject> finishedTask in finishedCompactBinaryResolves)
+			{
+				pendingCompactBinaryAttachments.Remove(finishedTask);
+			}
+		}
+
+		private static async Task ProcessContentIdsAsync(List<Attachment> list, List<Task<(ContentId, BlobId[]?)>> pendingContentIdResolves, List<BlobId> unresolvedBlobReferences)
+		{
+			// check for any content id resolves to finish
+			List<Task<(ContentId, BlobId[]?)>> finishedContentIdResolves = new();
+			foreach (Task<(ContentId, BlobId[]?)> pendingContentIdResolve in pendingContentIdResolves)
+			{
+				if (pendingContentIdResolve.IsCompleted)
+				{
+					ContentId? contentId = null;
+					BlobId[]? resolvedBlobs = null;
+					BlobId? blobIdentifier = null;
+					bool wasContentId = false;
+					try
+					{
+						(contentId, resolvedBlobs) = await pendingContentIdResolve;
+						blobIdentifier = contentId.AsBlobIdentifier();
+						wasContentId = !(resolvedBlobs is { Length: 1 } && resolvedBlobs[0].Equals(blobIdentifier));
+					}
+					catch (InvalidContentIdException)
+					{
+						resolvedBlobs = null;
+					}
+					catch (BlobNotFoundException e)
+					{
+						unresolvedBlobReferences.Add(e.Blob);
+					}
+
+					if (wasContentId && resolvedBlobs != null)
+					{
+						list.Add(new ContentIdAttachment(contentId!, resolvedBlobs));
+					}
+					else
+					{
+						list.Add(new BlobAttachment(blobIdentifier!));
+					}
+
+					finishedContentIdResolves.Add(pendingContentIdResolve);
+				}
+			}
+
+			// cleanup finished tasks
+			foreach (Task<(ContentId, BlobId[]?)> finishedTask in finishedContentIdResolves)
+			{
+				pendingContentIdResolves.Remove(finishedTask);
+			}
+		}
+
 		public async IAsyncEnumerable<BlobId> GetReferencedBlobsAsync(NamespaceId ns, CbObject cb, bool ignoreMissingBlobs = false, [EnumeratorCancellation] CancellationToken cancellationToken = default)
 		{
 			List<Task<(BlobId, bool)>> pendingBlobExistsChecks = new();
