Cloud DDC - Fixed spelling in docs

diff --git a/Engine/Source/Programs/UnrealCloudDDC/README.md b/Engine/Source/Programs/UnrealCloudDDC/README.md
--- a/Engine/Source/Programs/UnrealCloudDDC/README.md
+++ b/Engine/Source/Programs/UnrealCloudDDC/README.md
@@ -225,11 +225,11 @@ Namespace will share storage amongst all namespaces that share the same storage
 For DDC there is no need to customize any options in the namespace, everything defaults to a setup that is useful for ddc.
 
 ## Oplogs
-For oplogs we recommend enabling certain options on your namespace. The most important is the `GCMethod` which should be set to TTL (Time to live). This will make oplog entries automatically expire after 14 days (can be overriden by setting the `DefaultTTL`). Note that the uploaded content isn't removed after those 14 days, just the reference to it. Blob garbage collection will remove unused blobs once this happens. This makes sure old oplogs are automatically removed when not needed anymore (unlike the default GC method which keep anything that is used around which can quickly become a lot of storage for these large oplogs).
-The `storagePool` is something we recommend overriding to a seperate pool as well. This will expect you to have a S3 bucket / Azure blob storage that matched your configured name but with the storage pool suffix. So if your s3 bucket name is set to `ue-cloud-ddc-bucket` then the expected bucket for oplogs is `ue-cloud-ddc-bucket-oplog` e.g a `-oplog` suffix.
+For oplogs we recommend enabling certain options on your namespace. The most important is the `GCMethod` which should be set to TTL (Time to live). This will make oplog entries automatically expire after 14 days (can be overridden by setting the `DefaultTTL`). Note that the uploaded content isn't removed after those 14 days, just the reference to it. Blob garbage collection will remove unused blobs once this happens. This makes sure old oplogs are automatically removed when not needed anymore (unlike the default GC method which keep anything that is used around which can quickly become a lot of storage for these large oplogs).
+The `storagePool` is something we recommend overriding to a separate pool as well. This will expect you to have a S3 bucket / Azure blob storage that matched your configured name but with the storage pool suffix. So if your s3 bucket name is set to `ue-cloud-ddc-bucket` then the expected bucket for oplogs is `ue-cloud-ddc-bucket-oplog` e.g a `-oplog` suffix.
 
-Next you can enable the `UseBlobIndexForExists` option, this will use the metadata db (ScyllaDB) for exist checks rather then querying blob storage. This is much faster but puts more effort on your DB so make sure to have it scaled appopriately, Scylla usually handles this just fine. For large oplogs the existence checks can be a bottle neck.
-Lastly we have the `OnDemandReplication` option. When this is set that means any blob that is missing in a region is fetched on demand to other remote regions. This means that you do not nesscarilly need to have the speculative blob replication enabled but will be much slower for the first user that pulls down a oplog. We have that oplogs generally work okay with such a setup as long as you have enough users in a region. For smaller regions the speculative replication is better and it is always a net benefit to user experience (at a cost to transfer the data and store it).
+Next you can enable the `UseBlobIndexForExists` option, this will use the metadata db (ScyllaDB) for exist checks rather then querying blob storage. This is much faster but puts more effort on your DB so make sure to have it scaled appropriately, Scylla usually handles this just fine. For large oplogs the existence checks can be a bottle neck.
+Lastly we have the `OnDemandReplication` option. When this is set that means any blob that is missing in a region is fetched on demand to other remote regions. This means that you do not necessarily need to have the speculative blob replication enabled but will be much slower for the first user that pulls down a oplog. We have that oplogs generally work okay with such a setup as long as you have enough users in a region. For smaller regions the speculative replication is better and it is always a net benefit to user experience (at a cost to transfer the data and store it).
 
 This is an example of how the policies for a namespace could be configured for oplogs.
 ```
