Cloud DDC - Upgraded dotnet-monitor to 8.0.5

diff --git a/Engine/Source/Programs/UnrealCloudDDC/Helm/UnrealCloudDDC/templates/deployment.yaml b/Engine/Source/Programs/UnrealCloudDDC/Helm/UnrealCloudDDC/templates/deployment.yaml
--- a/Engine/Source/Programs/UnrealCloudDDC/Helm/UnrealCloudDDC/templates/deployment.yaml
+++ b/Engine/Source/Programs/UnrealCloudDDC/Helm/UnrealCloudDDC/templates/deployment.yaml
@@ -163,14 +163,27 @@ spec:
           {{- end }}
     {{- if .Values.useDotnetMonitorSidecar }}
       - name: sidecar
-        image: mcr.microsoft.com/dotnet/monitor:6.0.2
-        args: ["--urls", "http://*:52323", "--no-auth"]
+        image: mcr.microsoft.com/dotnet/monitor:8.0.5
+        env:
+        - name: DOTNETMONITOR_DiagnosticPort__ConnectionMode
+          value: Listen
+        - name: DOTNETMONITOR_Storage__DefaultSharedPath
+          value: /tmp
+        - name: DOTNETMONITOR_Urls
+          value: http://localhost:52323
+        - name: DOTNETMONITOR_Metrics__Endpoints
+          value: http://+:52325
+        args: ["collect", "--temp-apikey"]
         volumeMounts:
           - name: diagnostics
             mountPath: /tmp
         resources:
-          requests: 
-            cpu: 50m
+          requests:
+            memory: "32Mi"
+            cpu: "50m"
+          limits:
+            memory: "256Mi"
+            cpu: "250m"
     {{- end }}
     {{- if .Values.nginx.enabled }}
       - name: nginx
