Cloud DDC - Added alternative implementation of referencing resolving that should play better with the thread pool and is actually easier to follow as well. Not enabled by default as we want to see how this works under load.

diff --git a/Engine/Source/Programs/UnrealCloudDDC/Helm/UnrealCloudDDC/templates/configmap.yaml b/Engine/Source/Programs/UnrealCloudDDC/Helm/UnrealCloudDDC/templates/configmap.yaml
--- a/Engine/Source/Programs/UnrealCloudDDC/Helm/UnrealCloudDDC/templates/configmap.yaml
+++ b/Engine/Source/Programs/UnrealCloudDDC/Helm/UnrealCloudDDC/templates/configmap.yaml
@@ -110,6 +110,11 @@ ConsistencyCheck:
 {{ $config.ConsistencyCheck | toYaml | nindent 2 }}
 {{- end }}
 
+{{- if $config.ReferenceResolver }}
+ReferenceResolver: 
+{{ $config.ReferenceResolver | toYaml | nindent 2 }}
+{{- end }}
+
 {{- if $config.PayloadBuffering }}
 PayloadBuffering:
 {{ $config.PayloadBuffering | toYaml | nindent 2 }}
diff --git a/Engine/Source/Programs/UnrealCloudDDC/Jupiter/Implementation/IReferenceResolver.cs b/Engine/Source/Programs/UnrealCloudDDC/Jupiter/Implementation/IReferenceResolver.cs
--- a/Engine/Source/Programs/UnrealCloudDDC/Jupiter/Implementation/IReferenceResolver.cs
+++ b/Engine/Source/Programs/UnrealCloudDDC/Jupiter/Implementation/IReferenceResolver.cs
@@ -10,6 +10,7 @@ using EpicGames.Core;
 using EpicGames.Horde.Storage;
 using EpicGames.Serialization;
 using Jupiter.Utils;
+using Microsoft.Extensions.Options;
 using OpenTelemetry.Trace;
 
 namespace Jupiter.Implementation
@@ -95,16 +96,112 @@ namespace Jupiter.Implementation
 	{
 		private readonly IBlobService _blobStore;
 		private readonly IContentIdStore _contentIdStore;
+		private readonly IOptionsMonitor<ReferenceResolverSettings> _settings;
 		private readonly Tracer _tracer;
 
-		public ReferenceResolver(IBlobService blobStore, IContentIdStore contentIdStore, Tracer tracer)
+		public ReferenceResolver(IBlobService blobStore, IContentIdStore contentIdStore, IOptionsMonitor<ReferenceResolverSettings> settings, Tracer tracer)
 		{
 			_blobStore = blobStore;
 			_contentIdStore = contentIdStore;
+			_settings = settings;
 			_tracer = tracer;
 		}
 
 		public async IAsyncEnumerable<Attachment> GetAttachmentsAsync(NamespaceId ns, CbObject cb, [EnumeratorCancellation] CancellationToken cancellationToken)
+		{
+			if (_settings.CurrentValue.UseV2ReferenceResolving)
+			{
+				await foreach (Attachment attachment in GetAttachmentsV2Async(ns, cb, cancellationToken).Prefetch(_settings.CurrentValue.PrefetchCount, cancellationToken))
+				{
+					yield return attachment;
+				}
+			}
+			else
+			{
+				await foreach (Attachment attachment in GetAttachmentsOldAsync(ns, cb, cancellationToken))
+				{
+					yield return attachment;
+				}
+			}
+		}
+
+		private async IAsyncEnumerable<Attachment> GetAttachmentsV2Async(NamespaceId ns, CbObject cb, [EnumeratorCancellation] CancellationToken cancellationToken)
+		{
+			Queue<CbObject> objectsToVisit = new Queue<CbObject>();
+			objectsToVisit.Enqueue(cb);
+			List<BlobId> unresolvedBlobReferences = new List<BlobId>();
+
+			while (objectsToVisit.Count != 0)
+			{
+				if (objectsToVisit.TryDequeue(out CbObject? parent))
+				{
+					List<CbField> fields = new List<CbField>();
+					cb.IterateAttachments(field => fields.Add(field));
+					
+					foreach (CbField field in fields)
+					{
+						IoHash attachmentHash = field.AsAttachment();
+
+						if (field.IsBinaryAttachment())
+						{
+							ContentId? contentId = null;
+							BlobId[]? resolvedBlobs = null;
+							BlobId? blobIdentifier = null;
+							bool wasContentId = false;
+							try
+							{
+								(contentId, resolvedBlobs) = await ResolveContentIdAsync(ns, ContentId.FromIoHash(attachmentHash), cancellationToken);
+								blobIdentifier = contentId.AsBlobIdentifier();
+								wasContentId = !(resolvedBlobs is { Length: 1 } && resolvedBlobs[0].Equals(blobIdentifier));
+							}
+							catch (InvalidContentIdException)
+							{
+								resolvedBlobs = null;
+							}
+							catch (BlobNotFoundException e)
+							{
+								unresolvedBlobReferences.Add(e.Blob);
+							}
+
+							if (wasContentId && resolvedBlobs != null)
+							{
+								yield return new ContentIdAttachment(contentId!, resolvedBlobs);
+							}
+							else
+							{
+								yield return new BlobAttachment(blobIdentifier!);
+							}
+						}
+						else if (field.IsObjectAttachment())
+						{
+							BlobId blobIdentifier = BlobId.FromIoHash(attachmentHash);
+							yield return new ObjectAttachment(blobIdentifier);
+					
+							try
+							{
+								CbObject childBinaryObject = await ParseCompactBinaryAttachmentAsync(ns, blobIdentifier, cancellationToken);
+								objectsToVisit.Enqueue(childBinaryObject);
+							}
+							catch (BlobNotFoundException e)
+							{
+								unresolvedBlobReferences.Add(e.Blob);
+							}
+						}
+						else
+						{
+							throw new NotImplementedException($"Unknown attachment type for field {field}");
+						}
+					}
+				}
+			}
+
+			if (unresolvedBlobReferences.Count != 0)
+			{
+				throw new ReferenceIsMissingBlobsException(unresolvedBlobReferences);
+			}
+		}
+		
+		private async IAsyncEnumerable<Attachment> GetAttachmentsOldAsync(NamespaceId ns, CbObject cb, [EnumeratorCancellation] CancellationToken cancellationToken)
 		{
 			Queue<CbObject> objectsToVisit = new Queue<CbObject>();
 			objectsToVisit.Enqueue(cb);
@@ -390,6 +487,12 @@ namespace Jupiter.Implementation
 		}
 	}
 
+	public class ReferenceResolverSettings
+	{
+		public bool UseV2ReferenceResolving { get; set; } = false;
+		public int PrefetchCount { get; set; } = 64;
+	}
+
 	public class PartialReferenceResolveException : Exception
 	{
 		public List<ContentId> UnresolvedReferences { get; }
diff --git a/Engine/Source/Programs/UnrealCloudDDC/Jupiter/Jupiter.Startup.cs b/Engine/Source/Programs/UnrealCloudDDC/Jupiter/Jupiter.Startup.cs
--- a/Engine/Source/Programs/UnrealCloudDDC/Jupiter/Jupiter.Startup.cs
+++ b/Engine/Source/Programs/UnrealCloudDDC/Jupiter/Jupiter.Startup.cs
@@ -98,6 +98,7 @@ namespace Jupiter
 			services.AddOptions<SymbolsSettings>().Bind(Configuration.GetSection("Symbols")).ValidateDataAnnotations();
 
 			services.AddOptions<ConsistencyCheckSettings>().Bind(Configuration.GetSection("ConsistencyCheck")).ValidateDataAnnotations();
+			services.AddOptions<ReferenceResolverSettings>().Bind(Configuration.GetSection("ReferenceResolver")).ValidateDataAnnotations();
 
 			services.AddOptions<BufferedPayloadOptions>().Bind(Configuration.GetSection("PayloadBuffering")).ValidateDataAnnotations();
 
