Cloud DDC - Fixed a issue in v2 of the reference resolving as well as making it run more in parallel as expected by making consumers of the async enumerable actually consume it in parallel

diff --git a/Engine/Source/Programs/UnrealCloudDDC/Jupiter/Implementation/IReferenceResolver.cs b/Engine/Source/Programs/UnrealCloudDDC/Jupiter/Implementation/IReferenceResolver.cs
--- a/Engine/Source/Programs/UnrealCloudDDC/Jupiter/Implementation/IReferenceResolver.cs
+++ b/Engine/Source/Programs/UnrealCloudDDC/Jupiter/Implementation/IReferenceResolver.cs
@@ -136,7 +136,7 @@ namespace Jupiter.Implementation
 				if (objectsToVisit.TryDequeue(out CbObject? parent))
 				{
 					List<CbField> fields = new List<CbField>();
-					cb.IterateAttachments(field => fields.Add(field));
+					parent.IterateAttachments(field => fields.Add(field));
 					
 					foreach (CbField field in fields)
 					{
@@ -376,8 +376,6 @@ namespace Jupiter.Implementation
 
 		public async IAsyncEnumerable<BlobId> GetReferencedBlobsAsync(NamespaceId ns, CbObject cb, bool ignoreMissingBlobs = false, [EnumeratorCancellation] CancellationToken cancellationToken = default)
 		{
-			List<Task<(BlobId, bool)>> pendingBlobExistsChecks = new();
-			List<Task<(ContentIdAttachment, bool)>> pendingContentIdChecks = new();
 			List<ContentId> unresolvedContentIdReferences = new List<ContentId>();
 			List<BlobId> unresolvedBlobReferences = new List<BlobId>();
 
@@ -386,17 +384,44 @@ namespace Jupiter.Implementation
 			{
 				if (attachment is BlobAttachment blobAttachment)
 				{
-					pendingBlobExistsChecks.Add(CheckBlobExistsAsync(ns, blobAttachment.Identifier, cancellationToken));
+					(BlobId blob, bool exists) = await CheckBlobExistsAsync(ns, blobAttachment.Identifier, cancellationToken);
+					if (exists)
+					{
+						yield return blob;
+					}
+					else
+					{
+						unresolvedBlobReferences.Add(blob);
+					}
 				}
 				else if (attachment is ContentIdAttachment contentIdAttachment)
 				{
 					// If we find a content id we resolve that into the actual blobs it references
-					pendingContentIdChecks.Add(CheckContentIdExistsAsync(ns, contentIdAttachment, cancellationToken));
+					(ContentIdAttachment _, bool exists) = await CheckContentIdExistsAsync(ns, contentIdAttachment, cancellationToken);
+
+					if (!exists)
+					{
+						unresolvedContentIdReferences.Add(contentIdAttachment.Identifier);
+						continue;
+					}
+
+					foreach (BlobId b in contentIdAttachment.ReferencedBlobs)
+					{
+						yield return b;
+					}
 				}
 				else if (attachment is ObjectAttachment objectAttachment)
 				{
 					// a object just references the same blob, traversing the object attachment is done in GetAttachments
-					pendingBlobExistsChecks.Add(CheckBlobExistsAsync(ns, objectAttachment.Identifier, cancellationToken));
+					(BlobId blob, bool exists) = await CheckBlobExistsAsync(ns, objectAttachment.Identifier, cancellationToken);
+					if (exists)
+					{
+						yield return blob;
+					}
+					else
+					{
+						unresolvedBlobReferences.Add(blob);
+					}
 				}
 				else
 				{
@@ -404,36 +429,6 @@ namespace Jupiter.Implementation
 				}
 			}
 
-			// return any verified blobs
-			foreach (Task<(BlobId, bool)> pendingBlobExistsTask in pendingBlobExistsChecks)
-			{
-				(BlobId blob, bool exists) = await pendingBlobExistsTask;
-				if (exists)
-				{
-					yield return blob;
-				}
-				else
-				{
-					unresolvedBlobReferences.Add(blob);
-				}
-			}
-
-			foreach (Task<(ContentIdAttachment, bool)> pendingContentIdTask in pendingContentIdChecks)
-			{
-				(ContentIdAttachment contentIdAttachment, bool exists) = await pendingContentIdTask;
-
-				if (!exists)
-				{
-					unresolvedContentIdReferences.Add(contentIdAttachment.Identifier);
-					continue;
-				}
-
-				foreach (BlobId b in contentIdAttachment.ReferencedBlobs)
-				{
-					yield return b;
-				}
-			}
-
 			// if there were any content ids we did not recognize we throw a partial reference exception
 			if (!ignoreMissingBlobs && unresolvedContentIdReferences.Count != 0)
 			{
diff --git a/Engine/Source/Programs/UnrealCloudDDC/Jupiter/Implementation/RefService.cs b/Engine/Source/Programs/UnrealCloudDDC/Jupiter/Implementation/RefService.cs
--- a/Engine/Source/Programs/UnrealCloudDDC/Jupiter/Implementation/RefService.cs
+++ b/Engine/Source/Programs/UnrealCloudDDC/Jupiter/Implementation/RefService.cs
@@ -229,9 +229,6 @@ namespace Jupiter.Implementation
 			ContentId[] missingReferences = Array.Empty<ContentId>();
 			BlobId[] missingBlobs = Array.Empty<BlobId>();
 
-			List<Task> addToBucketTasks = new List<Task>();
-			List<Task> addRefMappingTasks = new List<Task>();
-
 			bool hasReferences = HasAttachments(payload);
 			if (hasReferences)
 			{
@@ -240,26 +237,25 @@ namespace Jupiter.Implementation
 				{
 					IAsyncEnumerable<BlobId> references = _referenceResolver.GetReferencedBlobsAsync(ns, payload, cancellationToken: cancellationToken);
 
-					await foreach (BlobId blobId in references)
+					await Parallel.ForEachAsync(references, cancellationToken, async (blobId, token) =>
 					{
+						Task blobIndexAddTask = _blobIndex.AddRefToBlobsAsync(ns, bucket, key, new BlobId[] { blobId }, cancellationToken);
+
 						if (_cloudDDCSettings.CurrentValue.EnableBucketStatsTracking)
 						{
-							addToBucketTasks.Add(Task.Run(async () =>
+							// if a blob is missing its not an error, the finalize will report this as missing and it will be uploaded and finalize ran again
+							try
+							{
+								BlobMetadata result = await _blobService.GetObjectMetadataAsync(ns, blobId, cancellationToken);
+								await _blobIndex.AddBlobToBucketListAsync(ns, bucket, key, blobId, result.Length, cancellationToken);
+							}
+							catch (BlobNotFoundException)
 							{
-								// if a blob is missing its not a error, the finalize will report this as missing and it will be uploaded and finalize ran again
-								try
-								{
-									BlobMetadata result = await _blobService.GetObjectMetadataAsync(ns, blobId, cancellationToken);
-									await _blobIndex.AddBlobToBucketListAsync(ns, bucket, key, blobId, result.Length, cancellationToken);
-								}
-								catch (BlobNotFoundException)
-								{
-								}
-							}, cancellationToken));
+							}
 						}
 
-						addRefMappingTasks.Add(_blobIndex.AddRefToBlobsAsync(ns, bucket, key, new BlobId[] { blobId }, cancellationToken));
-					}
+						await blobIndexAddTask;
+					});
 				}
 				catch (PartialReferenceResolveException e)
 				{
@@ -279,8 +275,6 @@ namespace Jupiter.Implementation
 				await _replicationLog.InsertAddEventAsync(ns, bucket, key, blobHash);
 			}
 
-			await Task.WhenAll(Task.WhenAll(addToBucketTasks), Task.WhenAll(addRefMappingTasks));
-
 			return (missingReferences, missingBlobs);
 		}
 
