Cloud DDC - Updates to v2 reference resolving, attempting to make it go wider when resolving the blob exists checks while still streaming the found attachments so we can issue the metadata requests for them.

diff --git a/Engine/Source/Programs/UnrealCloudDDC/Jupiter/Implementation/IReferenceResolver.cs b/Engine/Source/Programs/UnrealCloudDDC/Jupiter/Implementation/IReferenceResolver.cs
--- a/Engine/Source/Programs/UnrealCloudDDC/Jupiter/Implementation/IReferenceResolver.cs
+++ b/Engine/Source/Programs/UnrealCloudDDC/Jupiter/Implementation/IReferenceResolver.cs
@@ -111,7 +111,7 @@ namespace Jupiter.Implementation
 		{
 			if (_settings.CurrentValue.UseV2ReferenceResolving)
 			{
-				await foreach (Attachment attachment in GetAttachmentsV2Async(ns, cb, cancellationToken).Prefetch(_settings.CurrentValue.PrefetchCount, cancellationToken))
+				await foreach (Attachment attachment in GetAttachmentsV2Async(ns, cb, cancellationToken))
 				{
 					yield return attachment;
 				}
@@ -137,60 +137,72 @@ namespace Jupiter.Implementation
 				{
 					List<CbField> fields = new List<CbField>();
 					parent.IterateAttachments(field => fields.Add(field));
-					
-					foreach (CbField field in fields)
-					{
-						IoHash attachmentHash = field.AsAttachment();
 
-						if (field.IsBinaryAttachment())
+					foreach (CbField[] chunkedFields in fields.Chunk(_settings.CurrentValue.ChunkSize))
+					{
+						List<Attachment> foundAttachments = new List<Attachment>();
+						await Parallel.ForEachAsync(chunkedFields, cancellationToken, async (field, token) =>
 						{
-							ContentId? contentId = null;
-							BlobId[]? resolvedBlobs = null;
-							BlobId? blobIdentifier = null;
-							bool wasContentId = false;
-							try
-							{
-								(contentId, resolvedBlobs) = await ResolveContentIdAsync(ns, ContentId.FromIoHash(attachmentHash), cancellationToken);
-								blobIdentifier = contentId.AsBlobIdentifier();
-								wasContentId = !(resolvedBlobs is { Length: 1 } && resolvedBlobs[0].Equals(blobIdentifier));
-							}
-							catch (InvalidContentIdException)
+							IoHash attachmentHash = field.AsAttachment();
+
+							if (field.IsBinaryAttachment())
 							{
-								resolvedBlobs = null;
+								ContentId? contentId = null;
+								BlobId[]? resolvedBlobs = null;
+								BlobId? blobIdentifier = null;
+								bool wasContentId = false;
+								try
+								{
+									(contentId, resolvedBlobs) = await ResolveContentIdAsync(ns, ContentId.FromIoHash(attachmentHash), cancellationToken);
+									blobIdentifier = contentId.AsBlobIdentifier();
+									wasContentId = !(resolvedBlobs is { Length: 1 } && resolvedBlobs[0].Equals(blobIdentifier));
+								}
+								catch (InvalidContentIdException)
+								{
+									resolvedBlobs = null;
+								}
+								catch (BlobNotFoundException e)
+								{
+									unresolvedBlobReferences.Add(e.Blob);
+								}
+
+								if (wasContentId && resolvedBlobs != null)
+								{
+									foundAttachments.Add(new ContentIdAttachment(contentId!, resolvedBlobs));
+								}
+								else
+								{
+									foundAttachments.Add(new BlobAttachment(blobIdentifier!));
+								}
 							}
-							catch (BlobNotFoundException e)
+							else if (field.IsObjectAttachment())
 							{
-								unresolvedBlobReferences.Add(e.Blob);
-							}
 
-							if (wasContentId && resolvedBlobs != null)
-							{
-								yield return new ContentIdAttachment(contentId!, resolvedBlobs);
+								BlobId blobIdentifier = BlobId.FromIoHash(attachmentHash);
+								foundAttachments.Add(new ObjectAttachment(blobIdentifier));
+								
+								try
+								{
+									CbObject childBinaryObject = await ParseCompactBinaryAttachmentAsync(ns, blobIdentifier, cancellationToken);
+									objectsToVisit.Enqueue(childBinaryObject);
+								}
+								catch (BlobNotFoundException e)
+								{
+									unresolvedBlobReferences.Add(e.Blob);
+								}
 							}
 							else
 							{
-								yield return new BlobAttachment(blobIdentifier!);
-							}
-						}
-						else if (field.IsObjectAttachment())
-						{
-							BlobId blobIdentifier = BlobId.FromIoHash(attachmentHash);
-							yield return new ObjectAttachment(blobIdentifier);
-					
-							try
-							{
-								CbObject childBinaryObject = await ParseCompactBinaryAttachmentAsync(ns, blobIdentifier, cancellationToken);
-								objectsToVisit.Enqueue(childBinaryObject);
-							}
-							catch (BlobNotFoundException e)
-							{
-								unresolvedBlobReferences.Add(e.Blob);
+								throw new NotImplementedException($"Unknown attachment type for field {field}");
 							}
-						}
-						else
+						});
+
+						// using a parallel foreach above and consuming the results here means we do not stream the attachments as they are found instead we chunk the fields into pieces and stream the attachments for each piece while still going wide within the chunk
+
+						foreach (Attachment attachment in foundAttachments)
 						{
-							throw new NotImplementedException($"Unknown attachment type for field {field}");
-						}
+							yield return attachment;
+						}	
 					}
 				}
 			}
@@ -485,7 +497,7 @@ namespace Jupiter.Implementation
 	public class ReferenceResolverSettings
 	{
 		public bool UseV2ReferenceResolving { get; set; } = false;
-		public int PrefetchCount { get; set; } = 64;
+		public int ChunkSize { get; set; } = 64;
 	}
 
 	public class PartialReferenceResolveException : Exception
