Cloud DDC - Added more debugging information to understand occasional null dereference when resolving references using v2 method.

diff --git a/Engine/Source/Programs/UnrealCloudDDC/Jupiter/Implementation/IReferenceResolver.cs b/Engine/Source/Programs/UnrealCloudDDC/Jupiter/Implementation/IReferenceResolver.cs
--- a/Engine/Source/Programs/UnrealCloudDDC/Jupiter/Implementation/IReferenceResolver.cs
+++ b/Engine/Source/Programs/UnrealCloudDDC/Jupiter/Implementation/IReferenceResolver.cs
@@ -391,11 +391,13 @@ namespace Jupiter.Implementation
 			List<ContentId> unresolvedContentIdReferences = new List<ContentId>();
 			List<BlobId> unresolvedBlobReferences = new List<BlobId>();
 
+			bool attachmentsFound = false;
 			// Resolve all the attachments
-			await foreach (Attachment attachment in GetAttachmentsAsync(ns, cb, cancellationToken))
+			await foreach (Attachment? attachment in GetAttachmentsAsync(ns, cb, cancellationToken))
 			{
 				if (attachment is BlobAttachment blobAttachment)
 				{
+					attachmentsFound = true;
 					(BlobId blob, bool exists) = await CheckBlobExistsAsync(ns, blobAttachment.Identifier, cancellationToken);
 					if (exists)
 					{
@@ -408,6 +410,7 @@ namespace Jupiter.Implementation
 				}
 				else if (attachment is ContentIdAttachment contentIdAttachment)
 				{
+					attachmentsFound = true;
 					// If we find a content id we resolve that into the actual blobs it references
 					(ContentIdAttachment _, bool exists) = await CheckContentIdExistsAsync(ns, contentIdAttachment, cancellationToken);
 
@@ -424,6 +427,7 @@ namespace Jupiter.Implementation
 				}
 				else if (attachment is ObjectAttachment objectAttachment)
 				{
+					attachmentsFound = true;
 					// a object just references the same blob, traversing the object attachment is done in GetAttachments
 					(BlobId blob, bool exists) = await CheckBlobExistsAsync(ns, objectAttachment.Identifier, cancellationToken);
 					if (exists)
@@ -435,6 +439,10 @@ namespace Jupiter.Implementation
 						unresolvedBlobReferences.Add(blob);
 					}
 				}
+				else if (attachment == null)
+				{
+					throw new NotSupportedException($"Null attachment found when resolving attachments. Previous attachments found: {attachmentsFound}");
+				}
 				else
 				{
 					throw new NotSupportedException($"Unknown attachment type {attachment.GetType()}");
