Cloud DDC - Fixed race issue with reference resolver v2 (was not using a thread safe collection) and removed the added debugging code that isn't useful.

diff --git a/Engine/Source/Programs/UnrealCloudDDC/Jupiter/Implementation/IReferenceResolver.cs b/Engine/Source/Programs/UnrealCloudDDC/Jupiter/Implementation/IReferenceResolver.cs
--- a/Engine/Source/Programs/UnrealCloudDDC/Jupiter/Implementation/IReferenceResolver.cs
+++ b/Engine/Source/Programs/UnrealCloudDDC/Jupiter/Implementation/IReferenceResolver.cs
@@ -1,6 +1,7 @@
 // Copyright Epic Games, Inc. All Rights Reserved.
 
 using System;
+using System.Collections.Concurrent;
 using System.Collections.Generic;
 using System.Linq;
 using System.Runtime.CompilerServices;
@@ -140,7 +141,7 @@ namespace Jupiter.Implementation
 
 					foreach (CbField[] chunkedFields in fields.Chunk(_settings.CurrentValue.ChunkSize))
 					{
-						List<Attachment> foundAttachments = new List<Attachment>();
+						ConcurrentBag<Attachment> foundAttachments = new ConcurrentBag<Attachment>();
 						await Parallel.ForEachAsync(chunkedFields, cancellationToken, async (field, token) =>
 						{
 							IoHash attachmentHash = field.AsAttachment();
@@ -391,13 +392,11 @@ namespace Jupiter.Implementation
 			List<ContentId> unresolvedContentIdReferences = new List<ContentId>();
 			List<BlobId> unresolvedBlobReferences = new List<BlobId>();
 
-			bool attachmentsFound = false;
 			// Resolve all the attachments
-			await foreach (Attachment? attachment in GetAttachmentsAsync(ns, cb, cancellationToken))
+			await foreach (Attachment attachment in GetAttachmentsAsync(ns, cb, cancellationToken))
 			{
 				if (attachment is BlobAttachment blobAttachment)
 				{
-					attachmentsFound = true;
 					(BlobId blob, bool exists) = await CheckBlobExistsAsync(ns, blobAttachment.Identifier, cancellationToken);
 					if (exists)
 					{
@@ -410,7 +409,6 @@ namespace Jupiter.Implementation
 				}
 				else if (attachment is ContentIdAttachment contentIdAttachment)
 				{
-					attachmentsFound = true;
 					// If we find a content id we resolve that into the actual blobs it references
 					(ContentIdAttachment _, bool exists) = await CheckContentIdExistsAsync(ns, contentIdAttachment, cancellationToken);
 
@@ -427,7 +425,6 @@ namespace Jupiter.Implementation
 				}
 				else if (attachment is ObjectAttachment objectAttachment)
 				{
-					attachmentsFound = true;
 					// a object just references the same blob, traversing the object attachment is done in GetAttachments
 					(BlobId blob, bool exists) = await CheckBlobExistsAsync(ns, objectAttachment.Identifier, cancellationToken);
 					if (exists)
@@ -439,10 +436,6 @@ namespace Jupiter.Implementation
 						unresolvedBlobReferences.Add(blob);
 					}
 				}
-				else if (attachment == null)
-				{
-					throw new NotSupportedException($"Null attachment found when resolving attachments. Previous attachments found: {attachmentsFound}");
-				}
 				else
 				{
 					throw new NotSupportedException($"Unknown attachment type {attachment.GetType()}");
