CloudDDC - Added multipart download option from s3

diff --git a/Engine/Source/Programs/UnrealCloudDDC/Jupiter.Common/Implementation/StreamingBodyPayload.cs b/Engine/Source/Programs/UnrealCloudDDC/Jupiter.Common/Implementation/StreamingBodyPayload.cs
--- a/Engine/Source/Programs/UnrealCloudDDC/Jupiter.Common/Implementation/StreamingBodyPayload.cs
+++ b/Engine/Source/Programs/UnrealCloudDDC/Jupiter.Common/Implementation/StreamingBodyPayload.cs
@@ -99,6 +99,16 @@ namespace Jupiter.Common.Implementation
 		{
 			return new FilesystemBufferedPayloadWriter(filesystemTempPayloadRoot);
 		}
+
+		public FileInfo GetBackingFile()
+		{
+			if (_tempFile == null)
+			{
+				throw new Exception("Writable file was closed when fetching backing file");
+			}
+
+			return _tempFile;
+		}
 	}
 
 	/// <summary>
diff --git a/Engine/Source/Programs/UnrealCloudDDC/Jupiter/Implementation/Blob/AmazonS3Store.cs b/Engine/Source/Programs/UnrealCloudDDC/Jupiter/Implementation/Blob/AmazonS3Store.cs
--- a/Engine/Source/Programs/UnrealCloudDDC/Jupiter/Implementation/Blob/AmazonS3Store.cs
+++ b/Engine/Source/Programs/UnrealCloudDDC/Jupiter/Implementation/Blob/AmazonS3Store.cs
@@ -390,12 +390,17 @@ namespace Jupiter.Implementation
 			}
 		}
 
-		public async Task<BlobContents?> TryReadAsync(string path, LastAccessTrackingFlags flags = LastAccessTrackingFlags.DoTracking, CancellationToken cancellationToken = default)
+		private async Task<BlobContents?> GetMultipartAsync(string path, CancellationToken cancellationToken)
 		{
-			GetObjectResponse response;
+			using TransferUtility utility = new TransferUtility(_amazonS3);
 			try
 			{
-				response = await _amazonS3.GetObjectAsync(_bucketName, path, cancellationToken);
+				using FilesystemBufferedPayloadWriter writer = _payloadFactory.CreateFilesystemBufferedPayloadWriter();
+				TransferUtilityDownloadRequest downloadRequest = new TransferUtilityDownloadRequest { BucketName = _bucketName, Key = path, FilePath = writer.GetBackingFile().FullName };
+				await utility.DownloadAsync(downloadRequest, cancellationToken);
+
+				IBufferedPayload payload = writer.Done();
+				return new BlobContents(payload);
 			}
 			catch (AmazonS3Exception e)
 			{
@@ -408,9 +413,44 @@ namespace Jupiter.Implementation
 				{
 					return null;
 				}
+
+				if (e.StatusCode == HttpStatusCode.TooManyRequests)
+				{
+					throw new ResourceHasToManyRequestsException(e);
+				}
+
 				throw;
 			}
-			return new BlobContents(response.ResponseStream, response.ContentLength);
+		}
+
+		public async Task<BlobContents?> TryReadAsync(string path, LastAccessTrackingFlags flags = LastAccessTrackingFlags.DoTracking, CancellationToken cancellationToken = default)
+		{
+			if (_settings.CurrentValue.UseMultiPartDownload)
+			{
+				return await GetMultipartAsync(path, cancellationToken);
+			}
+			else
+			{
+				GetObjectResponse response;
+				try
+				{
+					response = await _amazonS3.GetObjectAsync(_bucketName, path, cancellationToken);
+				}
+				catch (AmazonS3Exception e)
+				{
+					if (e.ErrorCode == "NoSuchKey")
+					{
+						return null;
+					}
+
+					if (e.ErrorCode == "NoSuchBucket")
+					{
+						return null;
+					}
+					throw;
+				}
+				return new BlobContents(response.ResponseStream, response.ContentLength);
+			}
 		}
 
 		public async Task<bool> ExistsAsync(string path, CancellationToken cancellationToken)
diff --git a/Engine/Source/Programs/UnrealCloudDDC/Jupiter/Settings.cs b/Engine/Source/Programs/UnrealCloudDDC/Jupiter/Settings.cs
--- a/Engine/Source/Programs/UnrealCloudDDC/Jupiter/Settings.cs
+++ b/Engine/Source/Programs/UnrealCloudDDC/Jupiter/Settings.cs
@@ -198,6 +198,7 @@ namespace Jupiter
 
 		public bool? UseArnRegion { get; set; } = null;
 		public bool UseMultiPartUpload { get; set; } = false;
+		public bool UseMultiPartDownload { get; set; } = false;
 
 		/// <summary>
 		/// Allows you to override S3 behavior with chunk encoding, this needs to be set to false for uploads against GCS
