Cloud DDC - Fixed issue were large blobs that could not be inlined was not added to the blob replication log.

diff --git a/Engine/Source/Programs/UnrealCloudDDC/Jupiter/Implementation/Blob/BlobService.cs b/Engine/Source/Programs/UnrealCloudDDC/Jupiter/Implementation/Blob/BlobService.cs
--- a/Engine/Source/Programs/UnrealCloudDDC/Jupiter/Implementation/Blob/BlobService.cs
+++ b/Engine/Source/Programs/UnrealCloudDDC/Jupiter/Implementation/Blob/BlobService.cs
@@ -240,18 +240,9 @@ public class BlobService : IBlobService
 
 	private async Task<BlobId> PutObjectToStoresAsync(NamespaceId ns, byte[] payload, BlobId identifier, BucketId? bucketHint = null, bool? bypassCache = null, CancellationToken cancellationToken = default)
 	{
-		foreach (IBlobStore store in _blobStores)
-		{
-			await store.PutObjectAsync(ns, payload, identifier);
-		}
-
-		NamespacePolicy policy = _namespacePolicyResolver.GetPoliciesForNs(ns);
-		if (policy.PopulateFallbackNamespaceOnUpload && policy.FallbackNamespace.HasValue)
-		{
-			await PutObjectToStoresAsync(policy.FallbackNamespace.Value, payload, identifier, bucketHint: bucketHint, bypassCache: bypassCache, cancellationToken: cancellationToken);
-			await _blobIndex.AddBlobToIndexAsync(policy.FallbackNamespace.Value, identifier, cancellationToken: cancellationToken);
-		}
-		return identifier;
+		using MemoryStream ms = new MemoryStream(payload);
+		using IBufferedPayload bufferedPayload = await _bufferedPayloadFactory.CreateFromStreamAsync(ms, ms.Length, cancellationToken);
+		return await PutObjectToStoresAsync(ns, bufferedPayload, identifier, bucketHint, bypassCache, cancellationToken);
 	}
 
 	public async Task<BlobContents> GetObjectAsync(NamespaceId ns, BlobId blob, List<string>? storageLayers = null, bool supportsRedirectUri = false, bool allowOndemandReplication = true, BucketId? bucketHint = null, CancellationToken cancellationToken = default)
