Cloud DDC - Support nginx redirects for compact binarys used as refs, this is useful if it is a very large ref and thus is cached in the filesystem.

diff --git a/Engine/Source/Programs/UnrealCloudDDC/Jupiter/Controllers/ReferencesController.cs b/Engine/Source/Programs/UnrealCloudDDC/Jupiter/Controllers/ReferencesController.cs
--- a/Engine/Source/Programs/UnrealCloudDDC/Jupiter/Controllers/ReferencesController.cs
+++ b/Engine/Source/Programs/UnrealCloudDDC/Jupiter/Controllers/ReferencesController.cs
@@ -207,6 +207,11 @@ namespace Jupiter.Controllers
 				{
 					case CustomMediaTypeNames.UnrealCompactBinary:
 						{
+							if (_nginxRedirectHelper.CanRedirect(Request, blob))
+							{
+								return _nginxRedirectHelper.CreateActionResult(blob, responseType);
+							}
+
 							// for compact binary we can just serialize our internal object
 							await WriteBody(blob, CustomMediaTypeNames.UnrealCompactBinary);
 
