Cloud DDC - Added our own implementation of multipart get for s3, this will only apply to objects that in turn were uploaded using multipart.

diff --git a/Engine/Source/Programs/UnrealCloudDDC/Jupiter.Common/Implementation/StreamingBodyPayload.cs b/Engine/Source/Programs/UnrealCloudDDC/Jupiter.Common/Implementation/StreamingBodyPayload.cs
--- a/Engine/Source/Programs/UnrealCloudDDC/Jupiter.Common/Implementation/StreamingBodyPayload.cs
+++ b/Engine/Source/Programs/UnrealCloudDDC/Jupiter.Common/Implementation/StreamingBodyPayload.cs
@@ -99,16 +99,6 @@ namespace Jupiter.Common.Implementation
 		{
 			return new FilesystemBufferedPayloadWriter(filesystemTempPayloadRoot);
 		}
-
-		public FileInfo GetBackingFile()
-		{
-			if (_tempFile == null)
-			{
-				throw new Exception("Writable file was closed when fetching backing file");
-			}
-
-			return _tempFile;
-		}
 	}
 
 	/// <summary>
diff --git a/Engine/Source/Programs/UnrealCloudDDC/Jupiter/Implementation/Blob/AmazonS3Store.cs b/Engine/Source/Programs/UnrealCloudDDC/Jupiter/Implementation/Blob/AmazonS3Store.cs
--- a/Engine/Source/Programs/UnrealCloudDDC/Jupiter/Implementation/Blob/AmazonS3Store.cs
+++ b/Engine/Source/Programs/UnrealCloudDDC/Jupiter/Implementation/Blob/AmazonS3Store.cs
@@ -1,6 +1,7 @@
 // Copyright Epic Games, Inc. All Rights Reserved.
 
 using System;
+using System.Buffers;
 using System.Collections.Concurrent;
 using System.Collections.Generic;
 using System.IO;
@@ -392,15 +393,70 @@ namespace Jupiter.Implementation
 
 		private async Task<BlobContents?> GetMultipartAsync(string path, CancellationToken cancellationToken)
 		{
-			using TransferUtility utility = new TransferUtility(_amazonS3);
 			try
 			{
-				using FilesystemBufferedPayloadWriter writer = _payloadFactory.CreateFilesystemBufferedPayloadWriter();
-				TransferUtilityDownloadRequest downloadRequest = new TransferUtilityDownloadRequest { BucketName = _bucketName, Key = path, FilePath = writer.GetBackingFile().FullName };
-				await utility.DownloadAsync(downloadRequest, cancellationToken);
+				GetObjectRequest request = new GetObjectRequest { BucketName = _bucketName, Key = path, PartNumber = 1};
+				GetObjectResponse response = await _amazonS3.GetObjectAsync(request, cancellationToken);
+				// parts are only set of blobs that was uploaded using multipart, otherwise part 1 is the whole file
+				if (response.PartsCount > 1)
+				{
+					// object had parts, we download them in parallel
+					// -1 because parts start at 1
+					int partsCount = response.PartsCount.Value;
+
+					byte[][] parts = new byte[partsCount][];
+					long[] lengths = new long[partsCount];
+
+					try
+					{
+						{
+							byte[] partBuffer = ArrayPool<byte>.Shared.Rent((int)response.ContentLength);
+							await using MemoryStream ms = new MemoryStream(partBuffer);
+							await response.ResponseStream.CopyToAsync(ms, cancellationToken);
+							parts[0] = partBuffer;
+							lengths[0] = response.ContentLength;
+						}
 
-				IBufferedPayload payload = writer.Done();
-				return new BlobContents(payload);
+						// multipart object, download it in parts, we already have the first part
+						await Parallel.ForEachAsync(Enumerable.Range(2, partsCount-1), cancellationToken, async (part, token) =>
+						{
+							GetObjectRequest partRequest = new GetObjectRequest { BucketName = _bucketName, Key = path, PartNumber = part };
+							GetObjectResponse partResponse = await _amazonS3.GetObjectAsync(partRequest, cancellationToken);
+
+							byte[] partBuffer = ArrayPool<byte>.Shared.Rent((int)partResponse.ContentLength);
+							await using MemoryStream ms = new MemoryStream(partBuffer);
+							await partResponse.ResponseStream.CopyToAsync(ms, cancellationToken);
+
+							// the first part was retrieved in the initial request
+							parts[part-1] = partBuffer;
+							lengths[part-1] = partResponse.ContentLength;
+						});
+
+						using FilesystemBufferedPayloadWriter writer = _payloadFactory.CreateFilesystemBufferedPayloadWriter();
+
+						{
+							await using Stream s = writer.GetWritableStream();
+							for (int i = 0; i < parts.Length; i++)
+							{
+								byte[] part = parts[i];
+								await s.WriteAsync(part, 0, (int)lengths[i], cancellationToken);
+							}
+						}
+						
+						return new BlobContents(writer.Done());
+					}
+					finally
+					{
+						for (int i = 0; i < partsCount; i++)
+						{
+							ArrayPool<byte>.Shared.Return(parts[i]);
+						}
+					}
+				}
+				else
+				{
+					return new BlobContents(response.ResponseStream, response.ContentLength);
+				}
 			}
 			catch (AmazonS3Exception e)
 			{
