Cloud DDC - Enable multipart uploads by default for S3, we have been using this for a long while at Epic without any issues and it is a useful improvement for large files, esspecially worthwhile for oplogs.

diff --git a/Engine/Source/Programs/UnrealCloudDDC/Jupiter/Settings.cs b/Engine/Source/Programs/UnrealCloudDDC/Jupiter/Settings.cs
--- a/Engine/Source/Programs/UnrealCloudDDC/Jupiter/Settings.cs
+++ b/Engine/Source/Programs/UnrealCloudDDC/Jupiter/Settings.cs
@@ -197,7 +197,7 @@ namespace Jupiter
 		public Dictionary<string, string> StoragePoolBucketOverride { get; set; } = new Dictionary<string, string>();
 
 		public bool? UseArnRegion { get; set; } = null;
-		public bool UseMultiPartUpload { get; set; } = false;
+		public bool UseMultiPartUpload { get; set; } = true;
 		public bool UseMultiPartDownload { get; set; } = false;
 
 		/// <summary>
