Cloud DDC - Added namespace scopes to compressed blobs endpoints for easier filtering.

diff --git a/Engine/Source/Programs/UnrealCloudDDC/Jupiter/Controllers/CompressedBlobController.cs b/Engine/Source/Programs/UnrealCloudDDC/Jupiter/Controllers/CompressedBlobController.cs
--- a/Engine/Source/Programs/UnrealCloudDDC/Jupiter/Controllers/CompressedBlobController.cs
+++ b/Engine/Source/Programs/UnrealCloudDDC/Jupiter/Controllers/CompressedBlobController.cs
@@ -17,6 +17,7 @@ using Jupiter.Implementation;
 using Microsoft.AspNetCore.Authorization;
 using Microsoft.AspNetCore.Mvc;
 using Microsoft.Extensions.Primitives;
+using OpenTelemetry.Trace;
 using ContentId = Jupiter.Implementation.ContentId;
 
 namespace Jupiter.Controllers
@@ -63,6 +64,8 @@ namespace Jupiter.Controllers
 				return result;
 			}
 
+			Tracer.CurrentSpan.SetAttribute("namespace", ns.ToString());
+
 			try
 			{
 				(BlobContents blobContents, string mediaType) = await _storage.GetCompressedObjectAsync(ns, id, HttpContext.RequestServices, supportsRedirectUri: true);
@@ -228,6 +231,7 @@ namespace Jupiter.Controllers
 				return result;
 			}
 
+			Tracer.CurrentSpan.SetAttribute("namespace", ns.ToString());
 			_diagnosticContext.Set("Content-Length", Request.ContentLength ?? -1);
 
 			try
