Cloud DDC - Fixed issue building new versions of the benchmark docker container

diff --git a/Engine/Source/Programs/UnrealCloudDDC/Benchmarks/benchmark-suite/dockerfile b/Engine/Source/Programs/UnrealCloudDDC/Benchmarks/benchmark-suite/dockerfile
--- a/Engine/Source/Programs/UnrealCloudDDC/Benchmarks/benchmark-suite/dockerfile
+++ b/Engine/Source/Programs/UnrealCloudDDC/Benchmarks/benchmark-suite/dockerfile
@@ -9,7 +9,7 @@ COPY vegeta_12.8.4_linux_386.tar .
 RUN tar -xf vegeta_12.8.4_linux_386.tar
 
 COPY requirements.txt .
-RUN python3 -m pip install -r requirements.txt
+RUN python3 -m pip install -r requirements.txt --break-system-packages
 COPY *.py ./
 
 # Labels used by datadog to understand the log output we produce
