Cloud DDC - Fix issue were the blob cleanup would fail if no blobs had been written to the filesystem yet.

diff --git a/Engine/Source/Programs/UnrealCloudDDC/Jupiter/Implementation/Blob/FileSystemStore.cs b/Engine/Source/Programs/UnrealCloudDDC/Jupiter/Implementation/Blob/FileSystemStore.cs
--- a/Engine/Source/Programs/UnrealCloudDDC/Jupiter/Implementation/Blob/FileSystemStore.cs
+++ b/Engine/Source/Programs/UnrealCloudDDC/Jupiter/Implementation/Blob/FileSystemStore.cs
@@ -374,6 +374,11 @@ namespace Jupiter.Implementation
 		public IAsyncEnumerable<NamespaceId> ListNamespaces()
 		{
 			DirectoryInfo di = new DirectoryInfo(GetRootDir());
+			if (!di.Exists)
+			{
+				return AsyncEnumerable.Empty<NamespaceId>();
+			}
+
 			return di.GetDirectories().Select(x => new NamespaceId(x.Name)).ToAsyncEnumerable();
 		}
 	}
