CloudDDC - Send back the old (conflicting) object when a ref is not allowed to be overwritten so that the editor can give better errors when this happens.
Also fixed an issue were the full object was not considered when converting a ref into a octet stream (a path that is not actively used other then for testing).

diff --git a/Engine/Source/Programs/UnrealCloudDDC/Jupiter/Controllers/ReferencesController.cs b/Engine/Source/Programs/UnrealCloudDDC/Jupiter/Controllers/ReferencesController.cs
--- a/Engine/Source/Programs/UnrealCloudDDC/Jupiter/Controllers/ReferencesController.cs
+++ b/Engine/Source/Programs/UnrealCloudDDC/Jupiter/Controllers/ReferencesController.cs
@@ -37,7 +37,6 @@ using ContentId = Jupiter.Implementation.ContentId;
 
 namespace Jupiter.Controllers
 {
-	using BlobNotFoundException = Jupiter.Implementation.BlobNotFoundException;
 	using IDiagnosticContext = Serilog.IDiagnosticContext;
 
 	[ApiController]
@@ -226,7 +225,7 @@ namespace Jupiter.Controllers
 							{
 								int count = 0;
 								CbField? foundField = null;
-								cb.IterateAttachments(field =>
+								o.IterateAttachments(field =>
 								{
 									++count;
 									if (field.IsBinaryAttachment())
@@ -911,9 +910,22 @@ namespace Jupiter.Controllers
 					return Ok(new PutObjectResponse(missingArray));
 				}
 			}
-			catch (RefAlreadyExistsException)
+			catch (RefAlreadyExistsException e)
 			{
-				return Problem($"ref already exists {ns} {bucket} {key}", statusCode: (int)HttpStatusCode.Conflict);
+				byte[]? payload = e.OldObject.InlinePayload;
+				if (payload == null || payload.Length == 0)
+				{
+					using TelemetrySpan scope = _tracer.StartActiveSpan("blob.get").SetAttribute("operation.name", "blob.get");
+
+					await using BlobContents blobContents = await _blobStore.GetObjectAsync(ns, e.OldObject.BlobIdentifier, cancellationToken: HttpContext.RequestAborted); 
+					payload = await blobContents.Stream.ToByteArrayAsync(HttpContext.RequestAborted);
+				}
+				
+				// if a ref exists we send back the existing ref, always as a compact binary as that is how it is stored
+				Response.ContentType = CustomMediaTypeNames.UnrealCompactBinary;
+				Response.StatusCode = (int)HttpStatusCode.Conflict;
+				await Response.Body.WriteAsync(payload);
+				return new EmptyResult();
 			}
 		}
 
