Cloud DDC - Added option to set build store implementation to configmap, default to use scylla.

diff --git a/Engine/Source/Programs/UnrealCloudDDC/Helm/UnrealCloudDDC/templates/configmap-worker.yaml b/Engine/Source/Programs/UnrealCloudDDC/Helm/UnrealCloudDDC/templates/configmap-worker.yaml
--- a/Engine/Source/Programs/UnrealCloudDDC/Helm/UnrealCloudDDC/templates/configmap-worker.yaml
+++ b/Engine/Source/Programs/UnrealCloudDDC/Helm/UnrealCloudDDC/templates/configmap-worker.yaml
@@ -27,7 +27,8 @@ UnrealCloudDDC:
   ReplicationLogWriterImplementation:  {{ default "Scylla" $config.ReplicationLogWriterImplementation }}
   ContentIdStoreImplementation :  {{ default "Scylla" $config.ContentIdStoreImplementation }}
   BlobIndexImplementation :  {{ default "Scylla" $config.BlobIndexImplementation }}
-
+  BuildStoreImplementation :  {{ default "Scylla" $config.BuildStoreImplementation }}
+  
   # Disable last access tracking in the worker by default 
   # Typically operations in the worker do not bump the last accessed field so this just adds overhead
   # And you would typically not want to extend the lifetime of refs based on them being accessed by scheduled tasks
diff --git a/Engine/Source/Programs/UnrealCloudDDC/Helm/UnrealCloudDDC/templates/configmap.yaml b/Engine/Source/Programs/UnrealCloudDDC/Helm/UnrealCloudDDC/templates/configmap.yaml
--- a/Engine/Source/Programs/UnrealCloudDDC/Helm/UnrealCloudDDC/templates/configmap.yaml
+++ b/Engine/Source/Programs/UnrealCloudDDC/Helm/UnrealCloudDDC/templates/configmap.yaml
@@ -28,6 +28,7 @@ UnrealCloudDDC:
   ReplicationLogWriterImplementation:  {{ default "Scylla" $config.ReplicationLogWriterImplementation }}
   ContentIdStoreImplementation :  {{ default "Scylla" $config.ContentIdStoreImplementation }}
   BlobIndexImplementation :  {{ default "Scylla" $config.BlobIndexImplementation }}
+  BuildStoreImplementation :  {{ default "Scylla" $config.BuildStoreImplementation }}
 
 {{- if hasKey $config "EnableOnDemandReplication" }}
   EnableOnDemandReplication :  {{ $config.EnableOnDemandReplication }}
