Cloud DDC - Fixes for error when overwriting a ref whose object we could not resolve. Now produces a error message but also this happened because we were not fetching the blob from the db.

diff --git a/Engine/Source/Programs/UnrealCloudDDC/Jupiter/Controllers/ReferencesController.cs b/Engine/Source/Programs/UnrealCloudDDC/Jupiter/Controllers/ReferencesController.cs
--- a/Engine/Source/Programs/UnrealCloudDDC/Jupiter/Controllers/ReferencesController.cs
+++ b/Engine/Source/Programs/UnrealCloudDDC/Jupiter/Controllers/ReferencesController.cs
@@ -917,8 +917,17 @@ namespace Jupiter.Controllers
 				{
 					using TelemetrySpan scope = _tracer.StartActiveSpan("blob.get").SetAttribute("operation.name", "blob.get");
 
-					await using BlobContents blobContents = await _blobStore.GetObjectAsync(ns, e.OldObject.BlobIdentifier, cancellationToken: HttpContext.RequestAborted); 
-					payload = await blobContents.Stream.ToByteArrayAsync(HttpContext.RequestAborted);
+					try
+					{
+						await using BlobContents blobContents = await _blobStore.GetObjectAsync(ns, e.OldObject.BlobIdentifier, cancellationToken: HttpContext.RequestAborted); 
+						payload = await blobContents.Stream.ToByteArrayAsync(HttpContext.RequestAborted);
+					}
+					catch (BlobNotFoundException)
+					{
+						(CbObject errorObject, _) = ToErrorResult("Failed to find old ref during overwrite");
+						payload = errorObject.GetView().ToArray();
+					}
+
 				}
 				
 				// if a ref exists we send back the existing ref, always as a compact binary as that is how it is stored
diff --git a/Engine/Source/Programs/UnrealCloudDDC/Jupiter/Implementation/Objects/ScyllaReferencesStore.cs b/Engine/Source/Programs/UnrealCloudDDC/Jupiter/Implementation/Objects/ScyllaReferencesStore.cs
--- a/Engine/Source/Programs/UnrealCloudDDC/Jupiter/Implementation/Objects/ScyllaReferencesStore.cs
+++ b/Engine/Source/Programs/UnrealCloudDDC/Jupiter/Implementation/Objects/ScyllaReferencesStore.cs
@@ -164,7 +164,7 @@ namespace Jupiter.Implementation
 				bool insertRef = false;
 				try
 				{
-					RefRecord oldRecord = await GetAsync(ns, bucket, name, IReferencesStore.FieldFlags.None, IReferencesStore.OperationFlags.None, cancellationToken);
+					RefRecord oldRecord = await GetAsync(ns, bucket, name, IReferencesStore.FieldFlags.IncludePayload, IReferencesStore.OperationFlags.None, cancellationToken);
 
 					if (!oldRecord.BlobIdentifier.Equals(blobHash))
 					{
