Cloud DDC - Improved error messages when selecting build store implementation

diff --git a/Engine/Source/Programs/UnrealCloudDDC/Jupiter/Jupiter.Startup.cs b/Engine/Source/Programs/UnrealCloudDDC/Jupiter/Jupiter.Startup.cs
--- a/Engine/Source/Programs/UnrealCloudDDC/Jupiter/Jupiter.Startup.cs
+++ b/Engine/Source/Programs/UnrealCloudDDC/Jupiter/Jupiter.Startup.cs
@@ -231,7 +231,7 @@ namespace Jupiter
 				case UnrealCloudDDCSettings.BuildStoreImplementations.Scylla:
 					return ActivatorUtilities.CreateInstance<ScyllaBlockStore>(provider);
 				default:
-					throw new NotImplementedException($"Unknown build store implementation: {settings.BuildStoreImplementation}");
+					throw new NotImplementedException($"Unknown build index implementation: {settings.BuildStoreImplementation}");
 			}
 		}
 
@@ -245,7 +245,7 @@ namespace Jupiter
 				case UnrealCloudDDCSettings.BuildStoreImplementations.Scylla:
 					return ActivatorUtilities.CreateInstance<ScyllaBuildStore>(provider);
 				default:
-					throw new NotImplementedException($"Unknown blob store implementation: {settings.BuildStoreImplementation}");
+					throw new NotImplementedException($"Unknown build store implementation: {settings.BuildStoreImplementation}");
 			}
 		}
 
