Cloud DDC - Fixed some issues when converting from a json object to a compact binary object in the builds api

diff --git a/Engine/Source/Programs/Shared/EpicGames.Serialization/CbField.cs b/Engine/Source/Programs/Shared/EpicGames.Serialization/CbField.cs
--- a/Engine/Source/Programs/Shared/EpicGames.Serialization/CbField.cs
+++ b/Engine/Source/Programs/Shared/EpicGames.Serialization/CbField.cs
@@ -2738,7 +2738,11 @@ namespace EpicGames.Serialization
 		/// <returns></returns>
 		public static bool TryParse(string text, out CbObjectId objectId)
 		{
-			byte[] bytes = StringUtils.ParseHexString(text);
+			if (!StringUtils.TryParseHexString(text, out byte[]? bytes))
+			{
+				objectId = CbObjectId.Zero;
+				return false;
+			}
 			if (bytes.Length != 12)
 			{
 				objectId = CbObjectId.Zero;
diff --git a/Engine/Source/Programs/Shared/EpicGames.Serialization/CbSerializer.cs b/Engine/Source/Programs/Shared/EpicGames.Serialization/CbSerializer.cs
--- a/Engine/Source/Programs/Shared/EpicGames.Serialization/CbSerializer.cs
+++ b/Engine/Source/Programs/Shared/EpicGames.Serialization/CbSerializer.cs
@@ -245,6 +245,8 @@ namespace EpicGames.Serialization
 					{
 						ReadJsonField(prop.Value, prop.Name, writer);
 					}
+
+					writer.EndObject();
 					break;
 				}
 				case JsonValueKind.Array:
@@ -262,6 +264,8 @@ namespace EpicGames.Serialization
 					{
 						ReadJsonField(prop, null, writer);
 					}
+
+					writer.EndArray();
 					break;
 				}
 				case JsonValueKind.Null:
diff --git a/Engine/Source/Programs/UnrealCloudDDC/Jupiter/Implementation/Builds/IBlockStore.cs b/Engine/Source/Programs/UnrealCloudDDC/Jupiter/Implementation/Builds/IBlockStore.cs
--- a/Engine/Source/Programs/UnrealCloudDDC/Jupiter/Implementation/Builds/IBlockStore.cs
+++ b/Engine/Source/Programs/UnrealCloudDDC/Jupiter/Implementation/Builds/IBlockStore.cs
@@ -19,7 +19,7 @@ namespace Jupiter.Implementation
 	{
 		public List<string> MissingFields { get; init; }
 
-		public RequiredFieldMissingException(List<string> missingFields)
+		public RequiredFieldMissingException(List<string> missingFields) : base($"The following required fields were missing: {string.Join(",", missingFields)}")
 		{
 			MissingFields = missingFields;
 		}
