Cloud DDC - Do not expose non-finalized (partial) builds to the search

diff --git a/Engine/Source/Programs/UnrealCloudDDC/Jupiter/Controllers/BuildsController.cs b/Engine/Source/Programs/UnrealCloudDDC/Jupiter/Controllers/BuildsController.cs
--- a/Engine/Source/Programs/UnrealCloudDDC/Jupiter/Controllers/BuildsController.cs
+++ b/Engine/Source/Programs/UnrealCloudDDC/Jupiter/Controllers/BuildsController.cs
@@ -316,6 +316,11 @@ namespace Jupiter.Controllers
 			bool partialResult = false;
 			await foreach (BuildMetadata build in _buildStore.ListBuildsAsync(ns, bucket))
 			{
+				if (!build.IsFinalized)
+				{
+					// we ignore non-finalized builds
+					continue;
+				}
 				count++;
 				if (count < searchRequest.Options.Skip)
 				{
diff --git a/Engine/Source/Programs/UnrealCloudDDC/Jupiter/Implementation/Builds/IBuildStore.cs b/Engine/Source/Programs/UnrealCloudDDC/Jupiter/Implementation/Builds/IBuildStore.cs
--- a/Engine/Source/Programs/UnrealCloudDDC/Jupiter/Implementation/Builds/IBuildStore.cs
+++ b/Engine/Source/Programs/UnrealCloudDDC/Jupiter/Implementation/Builds/IBuildStore.cs
@@ -23,14 +23,16 @@ public interface IBuildStore
 
 public class BuildMetadata
 {
-	public BuildMetadata(CbObjectId buildId, CbObject metadata)
+	public BuildMetadata(CbObjectId buildId, CbObject metadata, bool isFinalized)
 	{
 		BuildId = buildId;
 		Metadata = metadata;
+		IsFinalized = isFinalized;
 	}
 
 	public CbObjectId BuildId { get; set; }
 	public CbObject Metadata { get; set; } = null!;
+	public bool IsFinalized { get; set; }
 }
 
 public class BuildRecord
diff --git a/Engine/Source/Programs/UnrealCloudDDC/Jupiter/Implementation/Builds/MemoryBuildStore.cs b/Engine/Source/Programs/UnrealCloudDDC/Jupiter/Implementation/Builds/MemoryBuildStore.cs
--- a/Engine/Source/Programs/UnrealCloudDDC/Jupiter/Implementation/Builds/MemoryBuildStore.cs
+++ b/Engine/Source/Programs/UnrealCloudDDC/Jupiter/Implementation/Builds/MemoryBuildStore.cs
@@ -67,7 +67,7 @@ public class MemoryBuildStore : IBuildStore
 				continue;
 			}
 
-			yield return new BuildMetadata(build.BuildId, build.BuildObject);
+			yield return new BuildMetadata(build.BuildId, build.BuildObject, build.IsFinalized);
 		}
 	}
 
diff --git a/Engine/Source/Programs/UnrealCloudDDC/Jupiter/Implementation/Builds/ScyllaBuildStore.cs b/Engine/Source/Programs/UnrealCloudDDC/Jupiter/Implementation/Builds/ScyllaBuildStore.cs
--- a/Engine/Source/Programs/UnrealCloudDDC/Jupiter/Implementation/Builds/ScyllaBuildStore.cs
+++ b/Engine/Source/Programs/UnrealCloudDDC/Jupiter/Implementation/Builds/ScyllaBuildStore.cs
@@ -183,7 +183,7 @@ public class ScyllaBuildEntry
 	public bool IsFinalized { get; set; }
 	public BuildMetadata ToBuildMetadata()
 	{
-		return new BuildMetadata(new CbObjectId(BuildId), new CbObject(BuildObject));
+		return new BuildMetadata(new CbObjectId(BuildId), new CbObject(BuildObject), IsFinalized);
 	}
 }
 
