Cloud DDC - Optimized finalize of build part by removing the double check of attachments (once for existence and once to determine parts are blocks and need extra handling). Should save about 20% of the wall time of submitting a part.

diff --git a/Engine/Source/Programs/UnrealCloudDDC/Jupiter/Controllers/BuildsController.cs b/Engine/Source/Programs/UnrealCloudDDC/Jupiter/Controllers/BuildsController.cs
--- a/Engine/Source/Programs/UnrealCloudDDC/Jupiter/Controllers/BuildsController.cs
+++ b/Engine/Source/Programs/UnrealCloudDDC/Jupiter/Controllers/BuildsController.cs
@@ -537,6 +537,20 @@ namespace Jupiter.Controllers
 				return Problem(e.Message, null, (int)HttpStatusCode.RequestTimeout);
 			}
 
+			BuildRecord? buildRecord = await _buildStore.GetBuildAsync(ns, bucket, buildId);
+			if (buildRecord == null)
+			{
+				throw new NotImplementedException($"No build found for namespace {ns} bucket {bucket} name {buildId}");
+			}
+
+			// build is finalized, we should not modify it
+			if (buildRecord.IsFinalized)
+			{
+				return BadRequest($"Build {buildId} is already finalized - it is read only.");
+			}
+
+			BlockContext blockContext = BlockContext.FromObject(buildRecord.BuildObject);
+
 			uint ttl = (uint)_namespacePolicyResolver.GetPoliciesForNs(ns).DefaultTTL.TotalSeconds;
 			await _buildStore.PutBuildPartAsync(ns, bucket, buildId, partId, partName, ttl);
 
@@ -546,18 +560,45 @@ namespace Jupiter.Controllers
 				{
 					using TelemetrySpan scope = _tracer.StartActiveSpan("ref.put").SetAttribute("operation.name", "ref.put");
 
-					(ContentId[] missingReferences, BlobId[] missingBlobs) = await _refService.PutAsync(ns, bucket, refId, payloadHash, payloadObject, cancellationToken: HttpContext.RequestAborted);
+					List<Task<BlobId?>> addBlockToContextTasks = new List<Task<BlobId?>>();
+					void OnBlobFound(BlobId blobId)
+					{
+						addBlockToContextTasks.Add(Task.Run(async () =>
+						{
+							BlobId? metadataBlob = await _blockStore.GetBlockMetadataAsync(ns, blobId);
+							if (metadataBlob != null)
+							{
+								return metadataBlob;
+							}
+
+							return null;
+						}));
+					}
+
+					(ContentId[] missingReferences, BlobId[] missingBlobs) = await _refService.PutAsync(ns, bucket, refId, payloadHash, payloadObject, OnBlobFound, cancellationToken: HttpContext.RequestAborted);
 
 					List<ContentHash> missingHashes = new List<ContentHash>(missingReferences);
 					missingHashes.AddRange(missingBlobs);
 					ContentHash[] missingArray = missingHashes.ToArray();
 					scope.SetAttribute("NeedsCount", missingArray.Length);
 
+					await Task.WhenAll(addBlockToContextTasks);
+
 					if (missingHashes.Count == 0)
 					{
 						using TelemetrySpan _ = _tracer.StartActiveSpan("build.finalize").SetAttribute("operation.name", "build.finalize");
 						
-						await DoFinalizeBuildPartAsync(ns, bucket, buildId, payloadObject, cancellationToken: HttpContext.RequestAborted);
+						List<BlobId> blockMetadataIds = new List<BlobId>();
+						foreach (Task<BlobId?> addBlockToContextTask in addBlockToContextTasks)
+						{
+							BlobId? b = await addBlockToContextTask;
+							if (b != null)
+							{
+								blockMetadataIds.Add(b);
+							}
+						}
+
+						await DoFinalizeBuildPartAsync(ns, blockContext, blockMetadataIds, cancellationToken: HttpContext.RequestAborted);
 					}
 
 					return Ok(new PutObjectResponse(missingArray));
@@ -584,25 +625,56 @@ namespace Jupiter.Controllers
 				Tracer.CurrentSpan.SetAttribute("bucket", bucket.ToString());
 				Tracer.CurrentSpan.SetAttribute("namespace", ns.ToString());
 
+				BuildRecord? buildRecord = await _buildStore.GetBuildAsync(ns, bucket, buildId);
+				if (buildRecord == null)
+				{
+					throw new NotImplementedException($"No build found for namespace {ns} bucket {bucket} name {buildId}");
+				}
+
+				// build is finalized, we should not modify it
+				if (buildRecord.IsFinalized)
+				{
+					return BadRequest($"Build {buildId} is already finalized - it is read only.");
+				}
+
+				BlockContext blockContext = BlockContext.FromObject(buildRecord.BuildObject);
+
+				List<Task<BlobId?>> addBlockToContextTasks = new List<Task<BlobId?>>();
+				void OnBlobFound(BlobId blobId)
+				{
+					addBlockToContextTasks.Add(Task.Run(async () =>
+					{
+						BlobId? metadataBlob = await _blockStore.GetBlockMetadataAsync(ns, blobId);
+						if (metadataBlob != null)
+						{
+							return metadataBlob;
+						}
+
+						return null;
+					}));
+				}
+
 				RefId refId = RefId.FromName($"{buildId}/{partId}");
-				(ContentId[] missingReferences, BlobId[] missingBlobs) = await _refService.FinalizeAsync(ns, bucket, refId, hash, HttpContext.RequestAborted);
+				(ContentId[] missingReferences, BlobId[] missingBlobs) = await _refService.FinalizeAsync(ns, bucket, refId, hash, OnBlobFound, HttpContext.RequestAborted);
 				List<ContentHash> missingHashes = new List<ContentHash>(missingReferences);
 				missingHashes.AddRange(missingBlobs);
 
+				await Task.WhenAll(addBlockToContextTasks);
 				if (missingHashes.Count == 0)
 				{
 					using TelemetrySpan _ = _tracer.StartActiveSpan("build.finalize").SetAttribute("operation.name", "build.finalize");
 
-					(RefRecord _, BlobContents? blob) = await _refService.GetAsync(ns, bucket, refId, doLastAccessTracking: false, cancellationToken: HttpContext.RequestAborted);
-					if (blob == null)
+					List<BlobId> blockMetadataIds = new List<BlobId>();
+					foreach (Task<BlobId?> addBlockToContextTask in addBlockToContextTasks)
 					{
-						throw new InvalidOperationException("No blob when attempting to finalize");
+						BlobId? b = await addBlockToContextTask;
+						if (b != null)
+						{
+							blockMetadataIds.Add(b);
+						}
 					}
 
-					byte[] blobContents = await blob.Stream.ToByteArrayAsync(HttpContext.RequestAborted);
-					CbObject partObject = new CbObject(blobContents);
-
-					await DoFinalizeBuildPartAsync(ns, bucket, buildId, partObject, HttpContext.RequestAborted);
+					await DoFinalizeBuildPartAsync(ns, blockContext, blockMetadataIds, cancellationToken: HttpContext.RequestAborted);
 				}
 
 				return Ok(new PutObjectResponse(missingHashes.ToArray()));
@@ -617,35 +689,13 @@ namespace Jupiter.Controllers
 			}
 		}
 
-		private async Task DoFinalizeBuildPartAsync(NamespaceId ns, BucketId bucket, CbObjectId buildId, CbObject partObject, CancellationToken cancellationToken)
+		private async Task DoFinalizeBuildPartAsync(NamespaceId ns, BlockContext blockContext, List<BlobId> blockMetadataIds, CancellationToken cancellationToken)
 		{
-			BuildRecord? buildRecord = await _buildStore.GetBuildAsync(ns, bucket, buildId);
-			if (buildRecord == null)
+			await Parallel.ForEachAsync(blockMetadataIds, cancellationToken, async (metadataBlobId, token) =>
 			{
-				throw new NotImplementedException($"No build found for namespace {ns} bucket {bucket} name {buildId}");
-			}
-
-			// no need to finalize again
-			if (buildRecord.IsFinalized)
-			{
-				return;
-			}
-
-			BlockContext blockContext = BlockContext.FromObject(buildRecord.BuildObject);
-			{
-				using TelemetrySpan _ = _tracer.StartActiveSpan("BuildService.ResolveReferences").SetAttribute("operation.name", "BuildService.ResolveReferences");
-				IAsyncEnumerable<BlobId> references = _referenceResolver.GetReferencedBlobsAsync(ns, partObject, cancellationToken: cancellationToken);
-
-				await Parallel.ForEachAsync(references, cancellationToken, async (blobId, token) =>
-				{
-					BlobId? metadataBlob = await _blockStore.GetBlockMetadataAsync(ns, blobId);
-					if (metadataBlob != null)
-					{
-						// this is a valid block because it has metadata submitted
-						await _blockStore.AddBlockToContextAsync(ns, blockContext, metadataBlob);
-					}
-				});
-			}
+				// this is a valid block because it has metadata submitted
+				await _blockStore.AddBlockToContextAsync(ns, blockContext, metadataBlobId);
+			});
 		}
 
 		#endregion
diff --git a/Engine/Source/Programs/UnrealCloudDDC/Jupiter/Controllers/ReferencesController.cs b/Engine/Source/Programs/UnrealCloudDDC/Jupiter/Controllers/ReferencesController.cs
--- a/Engine/Source/Programs/UnrealCloudDDC/Jupiter/Controllers/ReferencesController.cs
+++ b/Engine/Source/Programs/UnrealCloudDDC/Jupiter/Controllers/ReferencesController.cs
@@ -901,7 +901,7 @@ namespace Jupiter.Controllers
 				{
 					using TelemetrySpan scope = _tracer.StartActiveSpan("ref.put").SetAttribute("operation.name", "ref.put");
 
-					(ContentId[] missingReferences, BlobId[] missingBlobs) = await _refService.PutAsync(ns, bucket, key, blobHeader, payloadObject, cancellationToken: HttpContext.RequestAborted, allowOverwrite: allowOverwrite);
+					(ContentId[] missingReferences, BlobId[] missingBlobs) = await _refService.PutAsync(ns, bucket, key, blobHeader, payloadObject, allowOverwrite: allowOverwrite, cancellationToken: HttpContext.RequestAborted);
 
 					List<ContentHash> missingHashes = new List<ContentHash>(missingReferences);
 					missingHashes.AddRange(missingBlobs);
@@ -1018,7 +1018,7 @@ namespace Jupiter.Controllers
 				Tracer.CurrentSpan.SetAttribute("bucket", bucket.ToString());
 				Tracer.CurrentSpan.SetAttribute("namespace", ns.ToString());
 
-				(ContentId[] missingReferences, BlobId[] missingBlobs) = await _refService.FinalizeAsync(ns, bucket, key, hash, HttpContext.RequestAborted);
+				(ContentId[] missingReferences, BlobId[] missingBlobs) = await _refService.FinalizeAsync(ns, bucket, key, hash, cancellationToken: HttpContext.RequestAborted);
 				List<ContentHash> missingHashes = new List<ContentHash>(missingReferences);
 				missingHashes.AddRange(missingBlobs);
 
diff --git a/Engine/Source/Programs/UnrealCloudDDC/Jupiter/Implementation/IRefService.cs b/Engine/Source/Programs/UnrealCloudDDC/Jupiter/Implementation/IRefService.cs
--- a/Engine/Source/Programs/UnrealCloudDDC/Jupiter/Implementation/IRefService.cs
+++ b/Engine/Source/Programs/UnrealCloudDDC/Jupiter/Implementation/IRefService.cs
@@ -13,8 +13,8 @@ namespace Jupiter.Implementation
 	{
 		Task<(RefRecord, BlobContents?)> GetAsync(NamespaceId ns, BucketId bucket, RefId key, string[]? fields = null, bool doLastAccessTracking = true, CancellationToken cancellationToken = default);
 
-		Task<(ContentId[], BlobId[])> PutAsync(NamespaceId ns, BucketId bucket, RefId key, BlobId blobHash, CbObject payload, bool allowOverwrite = false, CancellationToken cancellationToken = default);
-		Task<(ContentId[], BlobId[])> FinalizeAsync(NamespaceId ns, BucketId bucket, RefId key, BlobId blobHash, CancellationToken cancellationToken = default);
+		Task<(ContentId[], BlobId[])> PutAsync(NamespaceId ns, BucketId bucket, RefId key, BlobId blobHash, CbObject payload, Action<BlobId>? onBlobFound = null, bool allowOverwrite = false, CancellationToken cancellationToken = default);
+		Task<(ContentId[], BlobId[])> FinalizeAsync(NamespaceId ns, BucketId bucket, RefId key, BlobId blobHash, Action<BlobId>? onBlobFound = null, CancellationToken cancellationToken = default);
 
 		IAsyncEnumerable<NamespaceId> GetNamespacesAsync(CancellationToken cancellationToken = default);
 		IAsyncEnumerable<RefId> GetRecordsInBucketAsync(NamespaceId ns, BucketId bucket, CancellationToken cancellationToken = default);
diff --git a/Engine/Source/Programs/UnrealCloudDDC/Jupiter/Implementation/RefService.cs b/Engine/Source/Programs/UnrealCloudDDC/Jupiter/Implementation/RefService.cs
--- a/Engine/Source/Programs/UnrealCloudDDC/Jupiter/Implementation/RefService.cs
+++ b/Engine/Source/Programs/UnrealCloudDDC/Jupiter/Implementation/RefService.cs
@@ -124,7 +124,7 @@ namespace Jupiter.Implementation
 			return (o, blobContents);
 		}
 
-		public async Task<(ContentId[], BlobId[])> PutAsync(NamespaceId ns, BucketId bucket, RefId key, BlobId blobHash, CbObject payload, bool allowOverwrite = false, CancellationToken cancellationToken = default)
+		public async Task<(ContentId[], BlobId[])> PutAsync(NamespaceId ns, BucketId bucket, RefId key, BlobId blobHash, CbObject payload, Action<BlobId>? onBlobFound = null, bool allowOverwrite = false, CancellationToken cancellationToken = default)
 		{
 			IServerTiming? serverTiming = _httpContextAccessor.HttpContext?.RequestServices.GetService<IServerTiming>();
 			using ServerTimingMetricScoped? serverTimingScope = serverTiming?.CreateServerTimingMetricScope("ref.put", "Inserting ref");
@@ -163,7 +163,7 @@ namespace Jupiter.Implementation
 				await blobStorePut;
 			}
 
-			return await DoFinalizeAsync(ns, bucket, key, blobHash, payload, cancellationToken);
+			return await DoFinalizeAsync(ns, bucket, key, blobHash, payload, onBlobFound, cancellationToken);
 		}
 
 		private bool HasAttachments(CbObject payload)
@@ -197,7 +197,7 @@ namespace Jupiter.Implementation
 			return payload.Any(FieldHasAttachments);
 		}
 
-		public async Task<(ContentId[], BlobId[])> FinalizeAsync(NamespaceId ns, BucketId bucket, RefId key, BlobId blobHash, CancellationToken cancellationToken = default)
+		public async Task<(ContentId[], BlobId[])> FinalizeAsync(NamespaceId ns, BucketId bucket, RefId key, BlobId blobHash, Action<BlobId>? onBlobFound = null, CancellationToken cancellationToken = default)
 		{
 			// finalize is intended to verify the state of the object in the db, so we bypass any caches to make sure the object we are working on is not stale
 			(RefRecord o, BlobContents? blob) = await GetAsync(ns, bucket, key, skipCache: true, cancellationToken: cancellationToken);
@@ -214,10 +214,10 @@ namespace Jupiter.Implementation
 				throw new ObjectHashMismatchException(ns, bucket, key, blobHash, o.BlobIdentifier);
 			}
 
-			return await DoFinalizeAsync(ns, bucket, key, blobHash, payload, cancellationToken);
+			return await DoFinalizeAsync(ns, bucket, key, blobHash, payload, onBlobFound, cancellationToken);
 		}
 
-		private async Task<(ContentId[], BlobId[])> DoFinalizeAsync(NamespaceId ns, BucketId bucket, RefId key, BlobId blobHash, CbObject payload, CancellationToken cancellationToken = default)
+		private async Task<(ContentId[], BlobId[])> DoFinalizeAsync(NamespaceId ns, BucketId bucket, RefId key, BlobId blobHash, CbObject payload, Action<BlobId>? onBlobFound = null, CancellationToken cancellationToken = default)
 		{
 			IServerTiming? serverTiming = _httpContextAccessor.HttpContext?.RequestServices.GetService<IServerTiming>();
 			using ServerTimingMetricScoped? serverTimingScope = serverTiming?.CreateServerTimingMetricScope("ref.finalize", "Finalizing the ref");
@@ -241,6 +241,7 @@ namespace Jupiter.Implementation
 					{
 						Task blobIndexAddTask = _blobIndex.AddRefToBlobsAsync(ns, bucket, key, new BlobId[] { blobId }, cancellationToken);
 
+						onBlobFound?.Invoke(blobId);
 						if (_cloudDDCSettings.CurrentValue.EnableBucketStatsTracking)
 						{
 							// if a blob is missing its not an error, the finalize will report this as missing and it will be uploaded and finalize ran again
