Cloud DDC - Made the ps1 script to allocate a token for cli debugging able to be run from any folder.

diff --git a/Engine/Source/Programs/UnrealCloudDDC/updateAuthHeaderUser.ps1 b/Engine/Source/Programs/UnrealCloudDDC/updateAuthHeaderUser.ps1
--- a/Engine/Source/Programs/UnrealCloudDDC/updateAuthHeaderUser.ps1
+++ b/Engine/Source/Programs/UnrealCloudDDC/updateAuthHeaderUser.ps1
@@ -1,5 +1,5 @@
 $TempFile = New-TemporaryFile
 
-& "dotnet" "..\..\..\..\Engine\Binaries\DotNET\OidcToken\portable\OidcToken.dll" "--Service=EpicGames-Okta" "--OutFile" $TempFile.FullName 
+& "dotnet" "$PSScriptRoot\..\..\..\..\Engine\Binaries\DotNET\OidcToken\portable\OidcToken.dll" "--Service=EpicGames-Okta" "--OutFile" $TempFile.FullName 
 $env:curl_auth_header = Get-Content -Raw -Path $TempFile.FullName | ConvertFrom-Json | ForEach-Object { "Authorization: Bearer $($_.Token)" }
-$env:auth_token = Get-Content -Raw -Path $TempFile.FullName | ConvertFrom-Json | ForEach-Object { "$($_.Token)" }
+$env:auth_token = Get-Content -Raw -Path $TempFile.FullName | ConvertFrom-Json | ForEach-Object { "$($_.Token)" }
\ No newline at end of file
