Cloud DDC - Added option 'RefFinalizeResolveMaxParallel' to manually bump the number of parallel finalize requests that are scheduled in an attempt to set it higher then number of threads (as this work is just waiting on Scylla anyway). This may allow us to make finalize of refs with a large number of attachments faster.

diff --git a/Engine/Source/Programs/UnrealCloudDDC/Jupiter/Implementation/RefService.cs b/Engine/Source/Programs/UnrealCloudDDC/Jupiter/Implementation/RefService.cs
--- a/Engine/Source/Programs/UnrealCloudDDC/Jupiter/Implementation/RefService.cs
+++ b/Engine/Source/Programs/UnrealCloudDDC/Jupiter/Implementation/RefService.cs
@@ -237,7 +237,7 @@ namespace Jupiter.Implementation
 				{
 					IAsyncEnumerable<BlobId> references = _referenceResolver.GetReferencedBlobsAsync(ns, payload, cancellationToken: cancellationToken);
 
-					await Parallel.ForEachAsync(references, cancellationToken, async (blobId, token) =>
+					await Parallel.ForEachAsync(references, new ParallelOptions {CancellationToken = cancellationToken, MaxDegreeOfParallelism = _cloudDDCSettings.CurrentValue.RefFinalizeResolveMaxParallel},async (blobId, token) =>
 					{
 						Task blobIndexAddTask = _blobIndex.AddRefToBlobsAsync(ns, bucket, key, new BlobId[] { blobId }, cancellationToken);
 
diff --git a/Engine/Source/Programs/UnrealCloudDDC/Jupiter/Settings.cs b/Engine/Source/Programs/UnrealCloudDDC/Jupiter/Settings.cs
--- a/Engine/Source/Programs/UnrealCloudDDC/Jupiter/Settings.cs
+++ b/Engine/Source/Programs/UnrealCloudDDC/Jupiter/Settings.cs
@@ -131,6 +131,11 @@ namespace Jupiter
 		public bool RequirePrivatePortForEnumeration { get; set; } = true;
 
 		public long InlineBlobMaxSize { get; set; } = 32 * 1024; // default to 32 kb blobs max
+
+		/// <summary>
+		/// Used to force a more parallel resolve than normal (which has no limits). Can be used to go wider then the number of threads which improves how quickly the finalize for a ref can be done, especially useful for refs with a lot of dependencies
+		/// </summary>
+		public int RefFinalizeResolveMaxParallel { get; set; } = -1;
 	}
 
 	public class MongoSettings
