Cloud DDC - Defaulted RefFinalizeResolveMaxParallel to its previous standard value of number of processors to avoid changing behavior. Can be set to -1 to allow finalize to go as wide as possisble and improve response times for large objects.

diff --git a/Engine/Source/Programs/UnrealCloudDDC/Jupiter/Settings.cs b/Engine/Source/Programs/UnrealCloudDDC/Jupiter/Settings.cs
--- a/Engine/Source/Programs/UnrealCloudDDC/Jupiter/Settings.cs
+++ b/Engine/Source/Programs/UnrealCloudDDC/Jupiter/Settings.cs
@@ -133,9 +133,9 @@ namespace Jupiter
 		public long InlineBlobMaxSize { get; set; } = 32 * 1024; // default to 32 kb blobs max
 
 		/// <summary>
-		/// Used to force a more parallel resolve than normal (which has no limits). Can be used to go wider then the number of threads which improves how quickly the finalize for a ref can be done, especially useful for refs with a lot of dependencies
+		/// Used to force a more parallel resolve than normal (which has no limits). Can be used to go wider then the number of processors which improves how quickly the finalize for a ref can be done, especially useful for refs with a lot of dependencies
 		/// </summary>
-		public int RefFinalizeResolveMaxParallel { get; set; } = -1;
+		public int RefFinalizeResolveMaxParallel { get; set; } = Environment.ProcessorCount;
 	}
 
 	public class MongoSettings
