Cloud DDC - Fixed incorrect tracing for blob cleanup, everything was flagged as a block. Only a logging issue, logic is correct.

diff --git a/Engine/Source/Programs/UnrealCloudDDC/Jupiter/Implementation/GC/OrphanBlobCleanupRefs.cs b/Engine/Source/Programs/UnrealCloudDDC/Jupiter/Implementation/GC/OrphanBlobCleanupRefs.cs
--- a/Engine/Source/Programs/UnrealCloudDDC/Jupiter/Implementation/GC/OrphanBlobCleanupRefs.cs
+++ b/Engine/Source/Programs/UnrealCloudDDC/Jupiter/Implementation/GC/OrphanBlobCleanupRefs.cs
@@ -127,7 +127,8 @@ namespace Jupiter.Implementation
 			string storagePoolName = string.IsNullOrEmpty(storagePool) ? "default" : storagePool;
 			using TelemetrySpan removeBlobScope = _tracer.StartActiveSpan("gc.blob")
 				.SetAttribute("operation.name", "gc.blob")
-				.SetAttribute("resource.name", $"{storagePoolName}.{blob}");
+				.SetAttribute("resource.name", $"{storagePoolName}.{blob}")
+				.SetAttribute("wasBlock", false.ToString());
 
 			bool found = false;
 
@@ -210,11 +211,11 @@ namespace Jupiter.Implementation
 
 			foreach (NamespaceId namespaceId in namespacesThatSharePool)
 			{
-				removeBlobScope.SetAttribute("wasBlock", true.ToString());
-
 				BlobId? metadata = await _blockStore.GetBlockMetadataAsync(namespaceId, blob);
 				if (metadata != null)
 				{
+					removeBlobScope.SetAttribute("wasBlock", true.ToString());
+
 					// this blob is a block, GC the block unique tracking
 					await _blockStore.DeleteBlockAsync(namespaceId, blob);
 				}
