Cloud DDC - Fixed issue were the block metadata was not being deleted when the block was being deleted.

diff --git a/Engine/Source/Programs/UnrealCloudDDC/Jupiter/Implementation/GC/OrphanBlobCleanupRefs.cs b/Engine/Source/Programs/UnrealCloudDDC/Jupiter/Implementation/GC/OrphanBlobCleanupRefs.cs
--- a/Engine/Source/Programs/UnrealCloudDDC/Jupiter/Implementation/GC/OrphanBlobCleanupRefs.cs
+++ b/Engine/Source/Programs/UnrealCloudDDC/Jupiter/Implementation/GC/OrphanBlobCleanupRefs.cs
@@ -217,7 +217,7 @@ namespace Jupiter.Implementation
 					removeBlobScope.SetAttribute("wasBlock", true.ToString());
 
 					// this blob is a block, GC the block unique tracking
-					await _blockStore.DeleteBlockAsync(namespaceId, blob);
+					await _blockStore.DeleteBlockAsync(namespaceId, metadata);
 				}
 			}
 
