Cloud DDC - Fixed exception thrown when GCing a block

diff --git a/Engine/Source/Programs/UnrealCloudDDC/Jupiter/Implementation/Builds/ScyllaBlockStore.cs b/Engine/Source/Programs/UnrealCloudDDC/Jupiter/Implementation/Builds/ScyllaBlockStore.cs
--- a/Engine/Source/Programs/UnrealCloudDDC/Jupiter/Implementation/Builds/ScyllaBlockStore.cs
+++ b/Engine/Source/Programs/UnrealCloudDDC/Jupiter/Implementation/Builds/ScyllaBlockStore.cs
@@ -95,7 +95,7 @@ public class ScyllaBlockStore : IBlockStore
 		{
 			if (string.Equals(blockEntry.Namespace, ns.ToString(), StringComparison.OrdinalIgnoreCase))
 			{
-				await _mapper.DeleteAsync<ScyllaBlockContextEntry>("WHERE namespace = ? AND block_context = ? AND block_id = ?", ns, blockEntry.BlockContext, blockIdentifier.HashData);
+				await _mapper.DeleteAsync<ScyllaBlockContextEntry>("WHERE namespace = ? AND block_context = ? AND block_id = ?", ns.ToString(), blockEntry.BlockContext, blockIdentifier.HashData);
 			}
 		}
 
