Cloud DDC - Enabled minio browser on port 9001 when using the tests compose, useful to be able to verify the contents after tests have run.

diff --git a/Engine/Source/Programs/UnrealCloudDDC/Composes/docker-compose-tests.yml b/Engine/Source/Programs/UnrealCloudDDC/Composes/docker-compose-tests.yml
--- a/Engine/Source/Programs/UnrealCloudDDC/Composes/docker-compose-tests.yml
+++ b/Engine/Source/Programs/UnrealCloudDDC/Composes/docker-compose-tests.yml
@@ -5,14 +5,15 @@ services:
     # Use to disable logging which can be spammy but also useful
     # logging:
     #  driver: none
-    command: "minio server /data"
+    command: "minio server /data --console-address :9001"
     environment:
     - "MINIO_ACCESS_KEY=FSYPTVCPKC2IV5D1U8RE"
     - "MINIO_SECRET_KEY=WaUf4Zq+J6YMlvbCVvTRkQJIPc7DvZwTm5UP0tPJ"
     networks:
     - network_bridge
     ports:
-    - "9000:9000" 
+    - "9000:9000"
+    - "9001:9001" 
   mongo:
     image: mongo
     # Use to disable logging which can be spammy but also useful
