Cloud DDC - Ignore incorrectly structured compact binary data in builds

diff --git a/Engine/Source/Programs/UnrealCloudDDC/Jupiter/Implementation/Builds/ScyllaBuildStore.cs b/Engine/Source/Programs/UnrealCloudDDC/Jupiter/Implementation/Builds/ScyllaBuildStore.cs
--- a/Engine/Source/Programs/UnrealCloudDDC/Jupiter/Implementation/Builds/ScyllaBuildStore.cs
+++ b/Engine/Source/Programs/UnrealCloudDDC/Jupiter/Implementation/Builds/ScyllaBuildStore.cs
@@ -1,5 +1,6 @@
 // Copyright Epic Games, Inc. All Rights Reserved.
 
+using System;
 using System.Collections.Generic;
 using System.Threading.Tasks;
 using Cassandra;
@@ -55,7 +56,13 @@ public class ScyllaBuildStore : IBuildStore
 		using TelemetrySpan scope = _tracer.BuildScyllaSpan("scylla.list_builds");
 		foreach (ScyllaBuildEntry buildEntry in await _mapper.FetchAsync<ScyllaBuildEntry>("WHERE namespace = ? AND bucket_id = ?", ns.ToString(), bucket.ToString()))
 		{
-			yield return buildEntry.ToBuildMetadata();
+			BuildMetadata? metadata = buildEntry.ToBuildMetadata();
+			// ignore invalid metadata
+			if (metadata == null)
+			{
+				continue;
+			}
+			yield return metadata;
 		}
 	}
 
@@ -181,9 +188,17 @@ public class ScyllaBuildEntry
 
 	[Cassandra.Mapping.Attributes.Column("is_finalized")]
 	public bool IsFinalized { get; set; }
-	public BuildMetadata ToBuildMetadata()
+	public BuildMetadata? ToBuildMetadata()
 	{
-		return new BuildMetadata(new CbObjectId(BuildId), new CbObject(BuildObject), IsFinalized);
+		try
+		{
+			return new BuildMetadata(new CbObjectId(BuildId), new CbObject(BuildObject), IsFinalized);
+		}
+		catch (IndexOutOfRangeException)
+		{
+			// incorrectly structured compact binaries are ignored
+			return null;
+		}
 	}
 }
 
