Cloud DDC - Added more debug logging around blocks to be able to better debug issues.

diff --git a/Engine/Source/Programs/UnrealCloudDDC/Jupiter/Controllers/BuildsController.cs b/Engine/Source/Programs/UnrealCloudDDC/Jupiter/Controllers/BuildsController.cs
--- a/Engine/Source/Programs/UnrealCloudDDC/Jupiter/Controllers/BuildsController.cs
+++ b/Engine/Source/Programs/UnrealCloudDDC/Jupiter/Controllers/BuildsController.cs
@@ -588,7 +588,7 @@ namespace Jupiter.Controllers
 
 					if (missingHashes.Count == 0)
 					{
-						using TelemetrySpan _ = _tracer.StartActiveSpan("build.finalize").SetAttribute("operation.name", "build.finalize");
+						using TelemetrySpan finalizeScope = _tracer.StartActiveSpan("build.finalize").SetAttribute("operation.name", "build.finalize");
 						
 						List<BlobId> blockMetadataIds = new List<BlobId>();
 						foreach (Task<BlobId?> addBlockToContextTask in addBlockToContextTasks)
@@ -599,7 +599,7 @@ namespace Jupiter.Controllers
 								blockMetadataIds.Add(b);
 							}
 						}
-
+						finalizeScope.SetAttribute("BlocksCount", blockMetadataIds.Count);
 						await DoFinalizeBuildPartAsync(ns, blockContext, blockMetadataIds, cancellationToken: HttpContext.RequestAborted);
 					}
 
@@ -641,6 +641,8 @@ namespace Jupiter.Controllers
 
 				BlockContext blockContext = BlockContext.FromObject(buildRecord.BuildObject);
 
+				Tracer.CurrentSpan.SetAttribute("blockContext", blockContext.ToString());
+
 				List<Task<BlobId?>> addBlockToContextTasks = new List<Task<BlobId?>>();
 				void OnBlobFound(BlobId blobId)
 				{
@@ -662,9 +664,10 @@ namespace Jupiter.Controllers
 				missingHashes.AddRange(missingBlobs);
 
 				await Task.WhenAll(addBlockToContextTasks);
+				Tracer.CurrentSpan.SetAttribute("NeedsCount", missingHashes.Count);
 				if (missingHashes.Count == 0)
 				{
-					using TelemetrySpan _ = _tracer.StartActiveSpan("build.finalize").SetAttribute("operation.name", "build.finalize");
+					using TelemetrySpan finalizeScope = _tracer.StartActiveSpan("build.finalize").SetAttribute("operation.name", "build.finalize");
 
 					List<BlobId> blockMetadataIds = new List<BlobId>();
 					foreach (Task<BlobId?> addBlockToContextTask in addBlockToContextTasks)
@@ -675,7 +678,7 @@ namespace Jupiter.Controllers
 							blockMetadataIds.Add(b);
 						}
 					}
-
+					finalizeScope.SetAttribute("BlocksCount", blockMetadataIds.Count);
 					await DoFinalizeBuildPartAsync(ns, blockContext, blockMetadataIds, cancellationToken: HttpContext.RequestAborted);
 				}
 
@@ -1068,6 +1071,8 @@ namespace Jupiter.Controllers
 			}
 
 			BlockContext context = BlockContext.FromObject(buildRecord.BuildObject, useBaseBranch: false);
+			Tracer.CurrentSpan.SetAttribute("blockContext", context.ToString());
+
 			IAsyncEnumerable<BlockMetadata> enumerable = _blockStore.ListBlockIndexAsync(ns, context);
 
 			CbWriter responseObject = new CbWriter();
@@ -1098,6 +1103,7 @@ namespace Jupiter.Controllers
 			{
 				context = BlockContext.FromObject(buildRecord.BuildObject, useBaseBranch: true);
 				enumerable = _blockStore.ListBlockIndexAsync(ns, context);
+				Tracer.CurrentSpan.SetAttribute("fallbackBlockContext", context.ToString());
 
 				await Parallel.ForEachAsync(enumerable.Skip(skip).Take(count), async (blockMetadata, token) =>
 				{
diff --git a/Engine/Source/Programs/UnrealCloudDDC/Jupiter/Implementation/Builds/ScyllaBlockStore.cs b/Engine/Source/Programs/UnrealCloudDDC/Jupiter/Implementation/Builds/ScyllaBlockStore.cs
--- a/Engine/Source/Programs/UnrealCloudDDC/Jupiter/Implementation/Builds/ScyllaBlockStore.cs
+++ b/Engine/Source/Programs/UnrealCloudDDC/Jupiter/Implementation/Builds/ScyllaBlockStore.cs
@@ -80,9 +80,10 @@ public class ScyllaBlockStore : IBlockStore
 
 		if (scyllaBlockIndexEntry != null)
 		{
+			scope.SetAttribute("found", bool.TrueString);
 			return new BlobId(scyllaBlockIndexEntry.MetadataBlobId);
 		}
-
+		scope.SetAttribute("found", bool.FalseString);
 		return null;
 	}
 
