Cloud DDC - Fixed issue with using index before it was initialized in scylla as reported in newer versions of scylla. Indexes should not be used directly in a select query and instead its used based on the columns that are used. With this we are not able to explicitly set the sort order which means we may be breaking the order that blocks are returned, at least this gets us unblocked for now while I consult with Scylla on the correct schema for this.

diff --git a/Engine/Source/Programs/UnrealCloudDDC/Jupiter/Implementation/Builds/ScyllaBlockStore.cs b/Engine/Source/Programs/UnrealCloudDDC/Jupiter/Implementation/Builds/ScyllaBlockStore.cs
--- a/Engine/Source/Programs/UnrealCloudDDC/Jupiter/Implementation/Builds/ScyllaBlockStore.cs
+++ b/Engine/Source/Programs/UnrealCloudDDC/Jupiter/Implementation/Builds/ScyllaBlockStore.cs
@@ -53,7 +53,7 @@ public class ScyllaBlockStore : IBlockStore
 		using TelemetrySpan scope = _tracer.BuildScyllaSpan("scylla.list_blocks");
 
 		foreach (ScyllaBlockContextEntry blockEntry in await _mapper.FetchAsync<ScyllaBlockContextEntry>(
-			         "SELECT * FROM block_context_by_time_index WHERE namespace = ? AND block_context = ? ORDER BY last_update DESC ", 
+			         "SELECT * FROM block_context WHERE namespace = ? AND block_context = ?", 
 			         ns.ToString(), blockContext.ToString()))
 		{
 			yield return new BlockMetadata { LastUpdate = blockEntry.LastUpdate, MetadataBlobId = new BlobId(blockEntry.BlockId) };
