Cloud DDC - Cleaning up old block context where blobs are missing during list blocks call. Not great to do cleanup here but for different reasons we have some stale data that shows up here. Also added tracing and logging to find out more information about list blocks somtimes losing its blocks. Lastly I fixed an issue were blocks could be missed during the finalize calls because we use a list but it was being added to by a concurrent callback, changed to thread safe container.

diff --git a/Engine/Source/Programs/UnrealCloudDDC/Jupiter/Controllers/BuildsController.cs b/Engine/Source/Programs/UnrealCloudDDC/Jupiter/Controllers/BuildsController.cs
--- a/Engine/Source/Programs/UnrealCloudDDC/Jupiter/Controllers/BuildsController.cs
+++ b/Engine/Source/Programs/UnrealCloudDDC/Jupiter/Controllers/BuildsController.cs
@@ -27,6 +27,7 @@ using System.Net;
 using System.Text;
 using System.Threading.Tasks;
 using System;
+using System.Collections.Concurrent;
 using System.Diagnostics;
 using System.Text.Json;
 using System.Text.Json.Serialization;
@@ -562,7 +563,7 @@ namespace Jupiter.Controllers
 				{
 					using TelemetrySpan scope = _tracer.StartActiveSpan("ref.put").SetAttribute("operation.name", "ref.put");
 
-					List<Task<BlobId?>> addBlockToContextTasks = new List<Task<BlobId?>>();
+					ConcurrentBag<Task<BlobId?>> addBlockToContextTasks = new ConcurrentBag<Task<BlobId?>>();
 					void OnBlobFound(BlobId blobId)
 					{
 						addBlockToContextTasks.Add(Task.Run(async () =>
@@ -643,7 +644,7 @@ namespace Jupiter.Controllers
 
 				Tracer.CurrentSpan.SetAttribute("blockContext", blockContext.ToString());
 
-				List<Task<BlobId?>> addBlockToContextTasks = new List<Task<BlobId?>>();
+				ConcurrentBag<Task<BlobId?>> addBlockToContextTasks = new ConcurrentBag<Task<BlobId?>>();
 				void OnBlobFound(BlobId blobId)
 				{
 					addBlockToContextTasks.Add(Task.Run(async () =>
@@ -1018,6 +1019,7 @@ namespace Jupiter.Controllers
 					return NotFound($"Unable to find block with id {blockId} as it does not exist, unable to upload metadata for it");
 				}
 
+				Tracer.CurrentSpan.SetAttribute("blockId", blockId.ToString());
 				await _blockStore.PutBlockMetadataAsync(ns, blockId, metadataBlockId);
 				await _blobIndex.AddBlobReferencesAsync(ns, blockId, metadataBlockId);
 
@@ -1071,7 +1073,6 @@ namespace Jupiter.Controllers
 			}
 
 			BlockContext context = BlockContext.FromObject(buildRecord.BuildObject, useBaseBranch: false);
-			Tracer.CurrentSpan.SetAttribute("blockContext", context.ToString());
 
 			IAsyncEnumerable<BlockMetadata> enumerable = _blockStore.ListBlockIndexAsync(ns, context);
 
@@ -1079,32 +1080,42 @@ namespace Jupiter.Controllers
 			responseObject.BeginObject();
 			responseObject.BeginArray("blocks");
 
+			int countOfBlocksFound = 0;
 			bool blocksFound = false;
-			await Parallel.ForEachAsync(enumerable.Skip(skip).Take(count), async (blockMetadata, token) =>
 			{
-				blocksFound = true;
-				try
+				using TelemetrySpan scope = _tracer.StartActiveSpan("list.blocks").SetAttribute("operation.name", "list.blocks");
+				scope.SetAttribute("blockContext", context.ToString());
+				await Parallel.ForEachAsync(enumerable.Skip(skip).Take(count), async (blockMetadata, token) =>
 				{
-					await using BlobContents contents = await _blobStore.GetObjectAsync(ns, blockMetadata.MetadataBlobId, cancellationToken: token);
-					byte[] b = await contents.Stream.ReadAllBytesAsync(cancellationToken: token);
-					CbObject o = new CbObject(b);
-					lock (responseObject)
+					blocksFound = true;
+					try
 					{
-						responseObject.WriteObject(o);
+						await using BlobContents contents = await _blobStore.GetObjectAsync(ns, blockMetadata.MetadataBlobId, cancellationToken: token);
+						byte[] b = await contents.Stream.ReadAllBytesAsync(cancellationToken: token);
+						CbObject o = new CbObject(b);
+						lock (responseObject)
+						{
+							responseObject.WriteObject(o);
+						}
+						Interlocked.Increment(ref countOfBlocksFound);
 					}
-				}
-				catch (BlobNotFoundException)
-				{
-					// ignore any blobs whose metadata we fail to retrieve
-				}
-			});
+					catch (BlobNotFoundException)
+					{
+						// we are unable to return any block that doesn't exist, we do clean this up here so that it's not found next time
+						await _blockStore.DeleteBlockAsync(ns, blockMetadata.MetadataBlobId);
+					}
+				});
+				scope.SetAttribute("countOfBlocksFound", countOfBlocksFound);
+			}
 
 			if (!blocksFound)
 			{
+				using TelemetrySpan scope = _tracer.StartActiveSpan("list.blocksFallback").SetAttribute("operation.name", "list.blocksFallback");
+				scope.SetAttribute("blockContext", context.ToString());
+
 				context = BlockContext.FromObject(buildRecord.BuildObject, useBaseBranch: true);
 				enumerable = _blockStore.ListBlockIndexAsync(ns, context);
-				Tracer.CurrentSpan.SetAttribute("fallbackBlockContext", context.ToString());
-
+				
 				await Parallel.ForEachAsync(enumerable.Skip(skip).Take(count), async (blockMetadata, token) =>
 				{
 					blocksFound = true;
@@ -1117,12 +1128,14 @@ namespace Jupiter.Controllers
 						{
 							responseObject.WriteObject(o);
 						}
+						Interlocked.Increment(ref countOfBlocksFound);
 					}
 					catch (BlobNotFoundException)
 					{
 						// ignore any blobs whose metadata we fail to retrieve
 					}
 				});
+				scope.SetAttribute("countOfBlocksFound", countOfBlocksFound);
 			}
 
 			responseObject.EndArray();
diff --git a/Engine/Source/Programs/UnrealCloudDDC/Jupiter/Implementation/Builds/IBlockStore.cs b/Engine/Source/Programs/UnrealCloudDDC/Jupiter/Implementation/Builds/IBlockStore.cs
--- a/Engine/Source/Programs/UnrealCloudDDC/Jupiter/Implementation/Builds/IBlockStore.cs
+++ b/Engine/Source/Programs/UnrealCloudDDC/Jupiter/Implementation/Builds/IBlockStore.cs
@@ -78,7 +78,7 @@ namespace Jupiter.Implementation
 	public interface IBlockStore
 	{
 		IAsyncEnumerable<BlockMetadata> ListBlockIndexAsync(NamespaceId ns, BlockContext blockContext);
-		Task AddBlockToContextAsync(NamespaceId ns, BlockContext blockContext, BlobId blockIdentifier);
+		Task AddBlockToContextAsync(NamespaceId ns, BlockContext blockContext, BlobId metadataBlockId);
 		Task PutBlockMetadataAsync(NamespaceId ns, BlobId blockIdentifier, BlobId metadataObjectId);
 		Task<BlobId?> GetBlockMetadataAsync(NamespaceId ns, BlobId blockIdentifier);
 		Task DeleteBlockAsync(NamespaceId ns, BlobId blockIdentifier);
diff --git a/Engine/Source/Programs/UnrealCloudDDC/Jupiter/Implementation/Builds/MemoryBlockStore.cs b/Engine/Source/Programs/UnrealCloudDDC/Jupiter/Implementation/Builds/MemoryBlockStore.cs
--- a/Engine/Source/Programs/UnrealCloudDDC/Jupiter/Implementation/Builds/MemoryBlockStore.cs
+++ b/Engine/Source/Programs/UnrealCloudDDC/Jupiter/Implementation/Builds/MemoryBlockStore.cs
@@ -27,17 +27,17 @@ public class MemoryBlockStore : IBlockStore
 		}
 	}
 
-	public Task AddBlockToContextAsync(NamespaceId ns, BlockContext blockContext, BlobId blockIdentifier)
+	public Task AddBlockToContextAsync(NamespaceId ns, BlockContext blockContext, BlobId metadataBlockId)
 	{
 		string key = BuildKey(ns, blockContext);
 
 		_blockContext.AddOrUpdate(key, _ =>
 		{
-			MemoryBlockContext o = new() { Blocks = { [blockIdentifier] = new BlockMetadata{LastUpdate = DateTime.Now, MetadataBlobId = blockIdentifier} } };
+			MemoryBlockContext o = new() { Blocks = { [metadataBlockId] = new BlockMetadata{LastUpdate = DateTime.Now, MetadataBlobId = metadataBlockId} } };
 			return o;
 		}, (s, o) =>
 		{
-			o.Blocks[blockIdentifier] = new BlockMetadata{LastUpdate = DateTime.Now, MetadataBlobId = blockIdentifier};
+			o.Blocks[metadataBlockId] = new BlockMetadata{LastUpdate = DateTime.Now, MetadataBlobId = metadataBlockId};
 			return o;
 		});
 
diff --git a/Engine/Source/Programs/UnrealCloudDDC/Jupiter/Implementation/Builds/ScyllaBlockStore.cs b/Engine/Source/Programs/UnrealCloudDDC/Jupiter/Implementation/Builds/ScyllaBlockStore.cs
--- a/Engine/Source/Programs/UnrealCloudDDC/Jupiter/Implementation/Builds/ScyllaBlockStore.cs
+++ b/Engine/Source/Programs/UnrealCloudDDC/Jupiter/Implementation/Builds/ScyllaBlockStore.cs
@@ -56,14 +56,14 @@ public class ScyllaBlockStore : IBlockStore
 			         "SELECT * FROM block_context WHERE namespace = ? AND block_context = ?", 
 			         ns.ToString(), blockContext.ToString()))
 		{
-			yield return new BlockMetadata { LastUpdate = blockEntry.LastUpdate, MetadataBlobId = new BlobId(blockEntry.BlockId) };
+			yield return new BlockMetadata { LastUpdate = blockEntry.LastUpdate, MetadataBlobId = new BlobId(blockEntry.MetadataBlobId) };
 		}
 	}
 
-	public async Task AddBlockToContextAsync(NamespaceId ns, BlockContext blockContext, BlobId blockIdentifier)
+	public async Task AddBlockToContextAsync(NamespaceId ns, BlockContext blockContext, BlobId metadataBlockId)
 	{
 		using TelemetrySpan scope = _tracer.BuildScyllaSpan("scylla.update_block_index");
-		await _mapper.InsertAsync<ScyllaBlockContextEntry>(new ScyllaBlockContextEntry(ns, blockContext, DateTime.Now, blockIdentifier));
+		await _mapper.InsertAsync<ScyllaBlockContextEntry>(new ScyllaBlockContextEntry(ns, blockContext, DateTime.Now, metadataBlockId));
 	}
 
 	public async Task PutBlockMetadataAsync(NamespaceId ns, BlobId blockIdentifier, BlobId metadataObjectId)
@@ -90,17 +90,18 @@ public class ScyllaBlockStore : IBlockStore
 	public async Task DeleteBlockAsync(NamespaceId ns, BlobId blockIdentifier)
 	{
 		using TelemetrySpan scope = _tracer.BuildScyllaSpan("scylla.delete_block_metadata");
-		Task deleteBlockIndexTask = _mapper.DeleteAsync<ScyllaBlockIndexEntry>("WHERE namespace = ? AND block_id = ?", ns.ToString(), blockIdentifier.HashData);
-
-		foreach (ScyllaBlockContextEntry blockEntry in await _mapper.FetchAsync<ScyllaBlockContextEntry>("SELECT * FROM block_context_by_blocks_index WHERE block_id = ? ", blockIdentifier.HashData))
+		int countOfContextsUsed = 0;
+		foreach (ScyllaBlockContextEntry blockEntry in await _mapper.FetchAsync<ScyllaBlockContextEntry>("SELECT * FROM block_context WHERE block_id = ? ", blockIdentifier.HashData))
 		{
 			if (string.Equals(blockEntry.Namespace, ns.ToString(), StringComparison.OrdinalIgnoreCase))
 			{
 				await _mapper.DeleteAsync<ScyllaBlockContextEntry>("WHERE namespace = ? AND block_context = ? AND block_id = ?", ns.ToString(), blockEntry.BlockContext, blockIdentifier.HashData);
+				countOfContextsUsed++;
 			}
 		}
 
-		await deleteBlockIndexTask;
+		scope.SetAttribute("UsedInCountOfScopes", countOfContextsUsed);
+		await _mapper.DeleteAsync<ScyllaBlockIndexEntry>("WHERE namespace = ? AND block_id = ?", ns.ToString(), blockIdentifier.HashData);
 	}
 }
 
@@ -140,7 +141,7 @@ public class ScyllaBlockContextEntry
 		Namespace = null!;
 		BlockContext = null!;
 		LastUpdate = DateTime.Now;
-		BlockId = null!;
+		MetadataBlobId = null!;
 	}
 
 	public ScyllaBlockContextEntry(NamespaceId @namespace, BlockContext context, DateTime lastUpdate, BlobId blockId)
@@ -148,7 +149,7 @@ public class ScyllaBlockContextEntry
 		Namespace = @namespace.ToString();
 		BlockContext = context.ToString();
 		LastUpdate = lastUpdate;
-		BlockId = blockId.HashData;
+		MetadataBlobId = blockId.HashData;
 	}
 
 	[Cassandra.Mapping.Attributes.PartitionKey]
@@ -163,5 +164,5 @@ public class ScyllaBlockContextEntry
 	public DateTime LastUpdate { get; set; }
 
 	[Cassandra.Mapping.Attributes.Column("block_id")]
-	public byte[] BlockId { get; set; }
+	public byte[] MetadataBlobId { get; set; }
 }
diff --git a/Engine/Source/Programs/UnrealCloudDDC/Jupiter/Implementation/GC/OrphanBlobCleanupRefs.cs b/Engine/Source/Programs/UnrealCloudDDC/Jupiter/Implementation/GC/OrphanBlobCleanupRefs.cs
--- a/Engine/Source/Programs/UnrealCloudDDC/Jupiter/Implementation/GC/OrphanBlobCleanupRefs.cs
+++ b/Engine/Source/Programs/UnrealCloudDDC/Jupiter/Implementation/GC/OrphanBlobCleanupRefs.cs
@@ -179,6 +179,8 @@ namespace Jupiter.Implementation
 					}
 					else if (baseBlobReference is BlobToBlobReference blobReference)
 					{
+						removeBlobScope.SetAttribute("hadBlobReference", bool.TrueString);
+
 						BlobId referringBlob = blobReference.Blob;
 						bool blobFound = await _blobService.ExistsAsync(blobNamespace, referringBlob, cancellationToken: cancellationToken);
 						if (blobFound)
