Cloud DDC - Fixed issue were block information was deleted as soon as the GC ran no matter if the block was referenced or not. Also did a speculative fix to avoid deleting rows from the block context as we are enumerating it.

diff --git a/Engine/Source/Programs/UnrealCloudDDC/Jupiter/Implementation/Builds/ScyllaBlockStore.cs b/Engine/Source/Programs/UnrealCloudDDC/Jupiter/Implementation/Builds/ScyllaBlockStore.cs
--- a/Engine/Source/Programs/UnrealCloudDDC/Jupiter/Implementation/Builds/ScyllaBlockStore.cs
+++ b/Engine/Source/Programs/UnrealCloudDDC/Jupiter/Implementation/Builds/ScyllaBlockStore.cs
@@ -90,17 +90,21 @@ public class ScyllaBlockStore : IBlockStore
 	public async Task DeleteBlockAsync(NamespaceId ns, BlobId blockIdentifier)
 	{
 		using TelemetrySpan scope = _tracer.BuildScyllaSpan("scylla.delete_block_metadata");
-		int countOfContextsUsed = 0;
+		List<string> foundContexts = new ();
 		foreach (ScyllaBlockContextEntry blockEntry in await _mapper.FetchAsync<ScyllaBlockContextEntry>("SELECT * FROM block_context WHERE block_id = ? ", blockIdentifier.HashData))
 		{
 			if (string.Equals(blockEntry.Namespace, ns.ToString(), StringComparison.OrdinalIgnoreCase))
 			{
-				await _mapper.DeleteAsync<ScyllaBlockContextEntry>("WHERE namespace = ? AND block_context = ? AND block_id = ?", ns.ToString(), blockEntry.BlockContext, blockIdentifier.HashData);
-				countOfContextsUsed++;
+				foundContexts.Add(blockEntry.BlockContext);
 			}
 		}
 
-		scope.SetAttribute("UsedInCountOfScopes", countOfContextsUsed);
+		foreach (string blockContext in foundContexts)
+		{
+			await _mapper.DeleteAsync<ScyllaBlockContextEntry>("WHERE namespace = ? AND block_context = ? AND block_id = ?", ns.ToString(), blockContext, blockIdentifier);
+		}
+
+		scope.SetAttribute("UsedInCountOfContexts", foundContexts.Count);
 		await _mapper.DeleteAsync<ScyllaBlockIndexEntry>("WHERE namespace = ? AND block_id = ?", ns.ToString(), blockIdentifier.HashData);
 	}
 }
diff --git a/Engine/Source/Programs/UnrealCloudDDC/Jupiter/Implementation/GC/OrphanBlobCleanupRefs.cs b/Engine/Source/Programs/UnrealCloudDDC/Jupiter/Implementation/GC/OrphanBlobCleanupRefs.cs
--- a/Engine/Source/Programs/UnrealCloudDDC/Jupiter/Implementation/GC/OrphanBlobCleanupRefs.cs
+++ b/Engine/Source/Programs/UnrealCloudDDC/Jupiter/Implementation/GC/OrphanBlobCleanupRefs.cs
@@ -211,18 +211,6 @@ namespace Jupiter.Implementation
 				return false;
 			}
 
-			foreach (NamespaceId namespaceId in namespacesThatSharePool)
-			{
-				BlobId? metadata = await _blockStore.GetBlockMetadataAsync(namespaceId, blob);
-				if (metadata != null)
-				{
-					removeBlobScope.SetAttribute("wasBlock", true.ToString());
-
-					// this blob is a block, GC the block unique tracking
-					await _blockStore.DeleteBlockAsync(namespaceId, metadata);
-				}
-			}
-
 			removeBlobScope.SetAttribute("removed", (!found).ToString());
 
 			_consideredBlobCounter.Add(1);
@@ -236,6 +224,20 @@ namespace Jupiter.Implementation
 			_deletedBlobCounter.Add(1);
 			_logger.LogInformation("GC Orphan blob {Blob} from {StoragePool} which was last modified at {LastModifiedTime}", blob, storagePoolName, lastModifiedTime);
 
+			// delete block information if this blob is a block
+			foreach (NamespaceId namespaceId in namespacesThatSharePool)
+			{
+				BlobId? metadata = await _blockStore.GetBlockMetadataAsync(namespaceId, blob);
+				if (metadata != null)
+				{
+					removeBlobScope.SetAttribute("wasBlock", true.ToString());
+
+					// this blob is a block, GC the block unique tracking
+					await _blockStore.DeleteBlockAsync(namespaceId, metadata);
+				}
+			}
+
+			// delete the blob itself
 			try
 			{
 				await _blobService.DeleteObjectAsync(namespacesThatSharePool, blob, cancellationToken);
