Cloud DDC - Trim away extra whitespace in claims when checking if what we get from IdP matches what we expect. Makes it less annoying to configure the IdP as it can be quite hard to spot a extra whitespace in your configuration for some of them.

diff --git a/Engine/Source/Programs/UnrealCloudDDC/Jupiter.Common/Authentication/Acl.cs b/Engine/Source/Programs/UnrealCloudDDC/Jupiter.Common/Authentication/Acl.cs
--- a/Engine/Source/Programs/UnrealCloudDDC/Jupiter.Common/Authentication/Acl.cs
+++ b/Engine/Source/Programs/UnrealCloudDDC/Jupiter.Common/Authentication/Acl.cs
@@ -166,13 +166,13 @@ public class AclPolicy
 				int separatorIndex = expectedClaim.IndexOf('=', StringComparison.InvariantCultureIgnoreCase);
 				string claimName = expectedClaim.Substring(0, separatorIndex);
 				string claimValue = expectedClaim.Substring(separatorIndex + 1);
-				if (user.HasClaim(claim => string.Equals(claim.Type, claimName, StringComparison.OrdinalIgnoreCase) && string.Equals(claim.Value, claimValue, StringComparison.OrdinalIgnoreCase)))
+				if (user.HasClaim(claim => string.Equals(claim.Type.Trim(), claimName.Trim(), StringComparison.OrdinalIgnoreCase) && string.Equals(claim.Value.Trim(), claimValue.Trim(), StringComparison.OrdinalIgnoreCase)))
 				{
 					claimFound = true;
 					logger.LogDebug("Found subset value claim {ClaimName} {ClaimValue} in policy \'{Name}\'.", claimName, claimValue, Name);
 				}
 			}
-			else if (user.HasClaim(claim => string.Equals(claim.Type, expectedClaim, StringComparison.OrdinalIgnoreCase)))
+			else if (user.HasClaim(claim => string.Equals(claim.Type.Trim(), expectedClaim.Trim(), StringComparison.OrdinalIgnoreCase)))
 			{
 				claimFound = true;
 				logger.LogDebug("Found exact claim match {ExpectedClaim} in policy \'{Name}\'.", expectedClaim, Name);
