Cloud DDC - Fixed incorrect type usage in scylla request when deleting block contents

diff --git a/Engine/Source/Programs/UnrealCloudDDC/Jupiter/Implementation/Builds/ScyllaBlockStore.cs b/Engine/Source/Programs/UnrealCloudDDC/Jupiter/Implementation/Builds/ScyllaBlockStore.cs
--- a/Engine/Source/Programs/UnrealCloudDDC/Jupiter/Implementation/Builds/ScyllaBlockStore.cs
+++ b/Engine/Source/Programs/UnrealCloudDDC/Jupiter/Implementation/Builds/ScyllaBlockStore.cs
@@ -101,7 +101,7 @@ public class ScyllaBlockStore : IBlockStore
 
 		foreach (string blockContext in foundContexts)
 		{
-			await _mapper.DeleteAsync<ScyllaBlockContextEntry>("WHERE namespace = ? AND block_context = ? AND block_id = ?", ns.ToString(), blockContext, blockIdentifier);
+			await _mapper.DeleteAsync<ScyllaBlockContextEntry>("WHERE namespace = ? AND block_context = ? AND block_id = ?", ns.ToString(), blockContext, blockIdentifier.HashData);
 		}
 
 		scope.SetAttribute("UsedInCountOfContexts", foundContexts.Count);
