Cloud DDC - Added server side sorting of list blocks so that we are likely to get the newer blocks before older ones. This should be done in Scylla but have not been able to get a index setup to do so, in discussions with Scylla about it.

diff --git a/Engine/Source/Programs/UnrealCloudDDC/Jupiter/Controllers/BuildsController.cs b/Engine/Source/Programs/UnrealCloudDDC/Jupiter/Controllers/BuildsController.cs
--- a/Engine/Source/Programs/UnrealCloudDDC/Jupiter/Controllers/BuildsController.cs
+++ b/Engine/Source/Programs/UnrealCloudDDC/Jupiter/Controllers/BuildsController.cs
@@ -1182,7 +1182,7 @@ namespace Jupiter.Controllers
 
 		[HttpGet("{ns}/{bucket}/{buildId}/parts/{partId}/blocks/listBlocks")]
 		[Produces(CustomMediaTypeNames.UnrealCompactBinary, MediaTypeNames.Application.Json)]
-		public async Task<IActionResult> ListBlockAsync(NamespaceId ns, BucketId bucket, CbObjectId buildId, string? format = null, [FromQuery] int count = 5000, [FromQuery] int skip = 0)
+		public async Task<IActionResult> ListBlockAsync(NamespaceId ns, BucketId bucket, CbObjectId buildId, string? format = null, [FromQuery] int count = 5000, [FromQuery] int skip = 0, [FromQuery] int maxConsidered = 25_000, [FromQuery] bool sort = true)
 		{
 			ActionResult? result = await _requestHelper.HasAccessToScopeAsync(User, Request, new AccessScope(ns, bucket), new[] { JupiterAclAction.EnumerateBucket });
 			if (result != null)
@@ -1196,70 +1196,72 @@ namespace Jupiter.Controllers
 				return NotFound($"Build {buildId} in bucket {bucket} and namespace {ns}");
 			}
 
-			BlockContext context = BlockContext.FromObject(buildRecord.BuildObject, useBaseBranch: false);
-
-			IAsyncEnumerable<BlockMetadata> enumerable = _blockStore.ListBlockIndexAsync(ns, context);
-
-			CbWriter responseObject = new CbWriter();
-			responseObject.BeginObject();
-			responseObject.BeginArray("blocks");
-
-			int countOfBlocksFound = 0;
-			bool blocksFound = false;
+			async Task<bool> WriteBlockMetadata(BlockContext context, CbWriter writer)
 			{
 				using TelemetrySpan scope = _tracer.StartActiveSpan("list.blocks").SetAttribute("operation.name", "list.blocks");
 				scope.SetAttribute("blockContext", context.ToString());
-				await Parallel.ForEachAsync(enumerable.Skip(skip).Take(count), async (blockMetadata, token) =>
+
+				// we sort these client side as a temporary solution as sorting these in Scylla has proven difficult
+				List<(DateTime, BlobId)> metadataBlobs = new();
+
+				int countOfBlocksFound = 0;
+				bool blocksFound = false;
 				{
-					blocksFound = true;
-					try
+					await foreach (BlockMetadata blockMetadata in _blockStore.ListBlockIndexAsync(ns, context).Take(maxConsidered))
 					{
-						await using BlobContents contents = await _blobStore.GetObjectAsync(ns, blockMetadata.MetadataBlobId, cancellationToken: token);
-						byte[] b = await contents.Stream.ReadAllBytesAsync(cancellationToken: token);
-						CbObject o = new CbObject(b);
-						lock (responseObject)
+						blocksFound = true;
+						lock (metadataBlobs)
 						{
-							responseObject.WriteObject(o);
+							metadataBlobs.Add((blockMetadata.LastUpdate, blockMetadata.MetadataBlobId));
 						}
-						Interlocked.Increment(ref countOfBlocksFound);
-					}
-					catch (BlobNotFoundException)
-					{
-						// we are unable to return any block that doesn't exist, we do clean this up here so that it's not found next time
-						await _blockStore.DeleteBlockAsync(ns, blockMetadata.MetadataBlobId);
 					}
-				});
-				scope.SetAttribute("countOfBlocksFound", countOfBlocksFound);
-			}
+				}
 
-			if (!blocksFound)
-			{
-				using TelemetrySpan scope = _tracer.StartActiveSpan("list.blocksFallback").SetAttribute("operation.name", "list.blocksFallback");
-				scope.SetAttribute("blockContext", context.ToString());
+				IEnumerable<(DateTime, BlobId)> enumerable = metadataBlobs.AsEnumerable();
+				if (sort)
+				{
+					enumerable = metadataBlobs.OrderByDescending(tuple => tuple.Item1);
+				}
 
-				context = BlockContext.FromObject(buildRecord.BuildObject, useBaseBranch: true);
-				enumerable = _blockStore.ListBlockIndexAsync(ns, context);
-				
-				await Parallel.ForEachAsync(enumerable.Skip(skip).Take(count), async (blockMetadata, token) =>
+				await Parallel.ForEachAsync(enumerable.Skip(skip).Take(count), async (pair, token) =>
 				{
-					blocksFound = true;
+					(DateTime _, BlobId blockMetadata) = pair;
 					try
 					{
-						await using BlobContents contents = await _blobStore.GetObjectAsync(ns, blockMetadata.MetadataBlobId, cancellationToken: token);
+						await using BlobContents contents = await _blobStore.GetObjectAsync(ns, blockMetadata, cancellationToken: token);
 						byte[] b = await contents.Stream.ReadAllBytesAsync(cancellationToken: token);
 						CbObject o = new CbObject(b);
-						lock (responseObject)
+						lock (writer)
 						{
-							responseObject.WriteObject(o);
+							writer.WriteObject(o);
 						}
+
 						Interlocked.Increment(ref countOfBlocksFound);
 					}
 					catch (BlobNotFoundException)
 					{
-						// ignore any blobs whose metadata we fail to retrieve
+						// we are unable to return any block that doesn't exist, we do clean this up here so that it's not found next time
+						await _blockStore.DeleteBlockAsync(ns, blockMetadata);
 					}
 				});
+
 				scope.SetAttribute("countOfBlocksFound", countOfBlocksFound);
+
+				return blocksFound;
+			}
+			
+			CbWriter responseObject = new CbWriter();
+			responseObject.BeginObject();
+			responseObject.BeginArray("blocks");
+
+			BlockContext context = BlockContext.FromObject(buildRecord.BuildObject, useBaseBranch: false);
+			bool blocksFound = await WriteBlockMetadata(context, responseObject);
+
+			if (!blocksFound)
+			{
+				// no blocks found in the branch, fallback to the base branch
+				BlockContext baseContext = BlockContext.FromObject(buildRecord.BuildObject, useBaseBranch: true);
+				await WriteBlockMetadata(baseContext, responseObject);
 			}
 
 			responseObject.EndArray();
