Cloud DDC - Added ability to enumerate which buckets exist in a namespace and which namespaces exists for builds, this is filtered down to only the buckets you have access to enumerate.
Also slightly changed response types when enumerating namespaces, if you have access to no namespaces the response is forbidden rather then simply getting nothing back (identical to what would have happened had you used any other endpoint and lacked access).

diff --git a/Engine/Source/Programs/UnrealCloudDDC/Jupiter.Common/BaseStartup.cs b/Engine/Source/Programs/UnrealCloudDDC/Jupiter.Common/BaseStartup.cs
--- a/Engine/Source/Programs/UnrealCloudDDC/Jupiter.Common/BaseStartup.cs
+++ b/Engine/Source/Programs/UnrealCloudDDC/Jupiter.Common/BaseStartup.cs
@@ -393,7 +393,7 @@ namespace Jupiter
 				string? assemblyName = Assembly.GetEntryAssembly()?.GetName().Name;
 				settings.SwaggerDoc("v1", info: new OpenApiInfo
 				{
-					Title = $"{assemblyName} API",
+					Title = "Unreal Cloud DDC",
 					Contact = new OpenApiContact
 					{
 						Name = "Joakim Lindqvist",
diff --git a/Engine/Source/Programs/UnrealCloudDDC/Jupiter/Controllers/BuildsController.cs b/Engine/Source/Programs/UnrealCloudDDC/Jupiter/Controllers/BuildsController.cs
--- a/Engine/Source/Programs/UnrealCloudDDC/Jupiter/Controllers/BuildsController.cs
+++ b/Engine/Source/Programs/UnrealCloudDDC/Jupiter/Controllers/BuildsController.cs
@@ -373,6 +373,7 @@ namespace Jupiter.Controllers
 			return Ok(new SearchResult(foundResults, partialResult));
 		}
 		#endregion
+
 		#region Build parts endpoints
 		[HttpGet("{ns}/{bucket}/{buildId}/parts/{partId}.{format?}")]
 		[Produces(CustomMediaTypeNames.UnrealCompressedBuffer, MediaTypeNames.Application.Json)]
@@ -1292,6 +1293,59 @@ namespace Jupiter.Controllers
 		}
 		
 		#endregion
+
+		#region Enumeration endpoints
+
+		[HttpGet("")]
+		[ProducesDefaultResponseType]
+		public async Task<IActionResult> GetNamespacesAsync()
+		{
+			NamespaceId[] namespaces = await _refService.GetNamespacesAsync(HttpContext.RequestAborted).ToArrayAsync(HttpContext.RequestAborted);
+
+			// filter namespaces down to only the namespaces the user has access to
+			List<NamespaceId> namespacesWithAccess = new();
+			foreach (NamespaceId ns in namespaces)
+			{
+				ActionResult? accessResult = await _requestHelper.HasAccessToScopeAsync(User, Request, new AccessScope(ns), new[] { JupiterAclAction.EnumerateBucket });
+				if (accessResult == null)
+				{
+					namespacesWithAccess.Add(ns);
+				}
+			}
+
+			if (!namespacesWithAccess.Any())
+			{
+				return new ForbidResult();
+			}
+
+			return Ok(new GetNamespacesResponse(namespacesWithAccess.ToArray()));
+		}
+
+		[HttpGet("{ns}")]
+		[ProducesDefaultResponseType]
+		public async Task<IActionResult> GetBucketsAsync(NamespaceId ns)
+		{
+			BucketId[] potentialBuckets = await _refService.GetBucketsAsync(ns, HttpContext.RequestAborted).ToArrayAsync(HttpContext.RequestAborted);
+
+			// filter buckets down to only the ones the user has access to
+			List<BucketId> verifiedBuckets = new();
+			foreach (BucketId bucket in potentialBuckets)
+			{
+				ActionResult? accessResult = await _requestHelper.HasAccessToScopeAsync(User, Request, new AccessScope(ns, bucket), new[] { JupiterAclAction.EnumerateBucket });
+				if (accessResult == null)
+				{
+					verifiedBuckets.Add(bucket);
+				}
+			}
+			if (!verifiedBuckets.Any())
+			{
+				return new ForbidResult();
+			}
+
+			return Ok(new GetBucketsResponse(verifiedBuckets.ToArray()));
+		}
+
+		#endregion
 	}
 
 	public class GetBlockMetadataRequest
diff --git a/Engine/Source/Programs/UnrealCloudDDC/Jupiter/Controllers/ReferencesController.cs b/Engine/Source/Programs/UnrealCloudDDC/Jupiter/Controllers/ReferencesController.cs
--- a/Engine/Source/Programs/UnrealCloudDDC/Jupiter/Controllers/ReferencesController.cs
+++ b/Engine/Source/Programs/UnrealCloudDDC/Jupiter/Controllers/ReferencesController.cs
@@ -102,6 +102,11 @@ namespace Jupiter.Controllers
 				}
 			}
 
+			if (!namespacesWithAccess.Any())
+			{
+				return new ForbidResult();
+			}
+
 			return Ok(new GetNamespacesResponse(namespacesWithAccess.ToArray()));
 		}
 
@@ -1731,6 +1736,17 @@ namespace Jupiter.Controllers
 		public NamespaceId[] Namespaces { get; set; }
 	}
 
+	public class GetBucketsResponse
+	{
+		[JsonConstructor]
+		public GetBucketsResponse(BucketId[] buckets)
+		{
+			Buckets = buckets;
+		}
+
+		public BucketId[] Buckets { get; set; }
+	}
+
 	public class HashMismatchException : Exception
 	{
 		public ContentHash SuppliedHash { get; }
diff --git a/Engine/Source/Programs/UnrealCloudDDC/Jupiter/Implementation/IRefService.cs b/Engine/Source/Programs/UnrealCloudDDC/Jupiter/Implementation/IRefService.cs
--- a/Engine/Source/Programs/UnrealCloudDDC/Jupiter/Implementation/IRefService.cs
+++ b/Engine/Source/Programs/UnrealCloudDDC/Jupiter/Implementation/IRefService.cs
@@ -17,6 +17,7 @@ namespace Jupiter.Implementation
 		Task<(ContentId[], BlobId[])> FinalizeAsync(NamespaceId ns, BucketId bucket, RefId key, BlobId blobHash, Action<BlobId>? onBlobFound = null, CancellationToken cancellationToken = default);
 
 		IAsyncEnumerable<NamespaceId> GetNamespacesAsync(CancellationToken cancellationToken = default);
+		IAsyncEnumerable<BucketId> GetBucketsAsync(NamespaceId ns, CancellationToken cancellationToken = default);
 		IAsyncEnumerable<RefId> GetRecordsInBucketAsync(NamespaceId ns, BucketId bucket, CancellationToken cancellationToken = default);
 
 		Task<bool> DeleteAsync(NamespaceId ns, BucketId bucket, RefId key, CancellationToken cancellationToken = default);
diff --git a/Engine/Source/Programs/UnrealCloudDDC/Jupiter/Implementation/RefService.cs b/Engine/Source/Programs/UnrealCloudDDC/Jupiter/Implementation/RefService.cs
--- a/Engine/Source/Programs/UnrealCloudDDC/Jupiter/Implementation/RefService.cs
+++ b/Engine/Source/Programs/UnrealCloudDDC/Jupiter/Implementation/RefService.cs
@@ -284,6 +284,11 @@ namespace Jupiter.Implementation
 			return _referencesStore.GetNamespacesAsync(cancellationToken);
 		}
 
+		public IAsyncEnumerable<BucketId> GetBucketsAsync(NamespaceId ns, CancellationToken cancellationToken = default)
+		{
+			return _referencesStore.GetBucketsAsync(ns, cancellationToken);
+		}
+
 		public IAsyncEnumerable<RefId> GetRecordsInBucketAsync(NamespaceId ns, BucketId bucket, CancellationToken cancellationToken = default)
 		{
 			return _referencesStore.GetRecordsInBucketAsync(ns, bucket, cancellationToken);
