Cloud DDC - Enabled multipart downloads in s3 by default as we have been using it a while and its working well (this is only used for blobs that are large enough that its worth).

diff --git a/Engine/Source/Programs/UnrealCloudDDC/Jupiter/Settings.cs b/Engine/Source/Programs/UnrealCloudDDC/Jupiter/Settings.cs
--- a/Engine/Source/Programs/UnrealCloudDDC/Jupiter/Settings.cs
+++ b/Engine/Source/Programs/UnrealCloudDDC/Jupiter/Settings.cs
@@ -210,7 +210,7 @@ namespace Jupiter
 
 		public bool? UseArnRegion { get; set; } = null;
 		public bool UseMultiPartUpload { get; set; } = true;
-		public bool UseMultiPartDownload { get; set; } = false;
+		public bool UseMultiPartDownload { get; set; } = true;
 
 		/// <summary>
 		/// Allows you to override S3 behavior with chunk encoding, this needs to be set to false for uploads against GCS
