Cloud DDC - Always require requests to opt in to getting a redirect uri response even if its allowed to be returned for the namespace policy. This makes sure that only clients that actually supports this gets a redirect.

diff --git a/Engine/Source/Programs/UnrealCloudDDC/Jupiter/Controllers/BlobController.cs b/Engine/Source/Programs/UnrealCloudDDC/Jupiter/Controllers/BlobController.cs
--- a/Engine/Source/Programs/UnrealCloudDDC/Jupiter/Controllers/BlobController.cs
+++ b/Engine/Source/Programs/UnrealCloudDDC/Jupiter/Controllers/BlobController.cs
@@ -53,6 +53,7 @@ namespace Jupiter.Controllers
 			[Required] NamespaceId ns,
 			[Required] BlobId id,
 			[FromQuery] List<string>? storageLayers = null,
+			[FromQuery] bool supportsRedirect = false,
 			[FromQuery] bool allowOndemandReplication = true)
 		{
 			ActionResult? result = await _requestHelper.HasAccessToNamespaceAsync(User, Request, ns, new[] { JupiterAclAction.ReadObject });
@@ -63,7 +64,7 @@ namespace Jupiter.Controllers
 
 			try
 			{
-				BlobContents blobContents = await GetImplAsync(ns, id, storageLayers, supportsRedirectUri: true, allowOndemandReplication: allowOndemandReplication);
+				BlobContents blobContents = await GetImplAsync(ns, id, storageLayers, supportsRedirectUri: supportsRedirect, allowOndemandReplication: allowOndemandReplication);
 
 				if (blobContents.RedirectUri != null)
 				{
diff --git a/Engine/Source/Programs/UnrealCloudDDC/Jupiter/Controllers/CompressedBlobController.cs b/Engine/Source/Programs/UnrealCloudDDC/Jupiter/Controllers/CompressedBlobController.cs
--- a/Engine/Source/Programs/UnrealCloudDDC/Jupiter/Controllers/CompressedBlobController.cs
+++ b/Engine/Source/Programs/UnrealCloudDDC/Jupiter/Controllers/CompressedBlobController.cs
@@ -56,7 +56,8 @@ namespace Jupiter.Controllers
 
 		public async Task<IActionResult> GetAsync(
 			[Required] NamespaceId ns,
-			[Required] ContentId id)
+			[Required] ContentId id,
+			[FromQuery] bool supportsRedirect = false)
 		{
 			ActionResult? result = await _requestHelper.HasAccessToNamespaceAsync(User, Request, ns, new[] { JupiterAclAction.ReadObject });
 			if (result != null)
@@ -68,7 +69,7 @@ namespace Jupiter.Controllers
 
 			try
 			{
-				(BlobContents blobContents, string mediaType) = await _storage.GetCompressedObjectAsync(ns, id, HttpContext.RequestServices, supportsRedirectUri: true);
+				(BlobContents blobContents, string mediaType) = await _storage.GetCompressedObjectAsync(ns, id, HttpContext.RequestServices, supportsRedirectUri: supportsRedirect);
 
 				StringValues acceptHeader = Request.Headers["Accept"];
 				if (!acceptHeader.Contains("*/*") && acceptHeader.Count != 0 && !acceptHeader.Contains(mediaType))
diff --git a/Engine/Source/Programs/UnrealCloudDDC/Jupiter/Controllers/ReferencesController.cs b/Engine/Source/Programs/UnrealCloudDDC/Jupiter/Controllers/ReferencesController.cs
--- a/Engine/Source/Programs/UnrealCloudDDC/Jupiter/Controllers/ReferencesController.cs
+++ b/Engine/Source/Programs/UnrealCloudDDC/Jupiter/Controllers/ReferencesController.cs
@@ -1050,7 +1050,8 @@ namespace Jupiter.Controllers
 			[FromRoute][Required] BucketId bucket,
 			[FromRoute][Required] RefId key,
 #pragma warning restore IDE0060
-			[Required] ContentId id)
+			[Required] ContentId id,
+			[FromQuery] bool supportsRedirect = false)
 		{
 			ActionResult? result = await _requestHelper.HasAccessToNamespaceAsync(User, Request, ns, new[] { JupiterAclAction.ReadObject });
 			if (result != null)
@@ -1060,7 +1061,7 @@ namespace Jupiter.Controllers
 
 			try
 			{
-				(BlobContents blobContents, string mediaType) = await _blobStore.GetCompressedObjectAsync(ns, id, HttpContext.RequestServices, supportsRedirectUri: true);
+				(BlobContents blobContents, string mediaType) = await _blobStore.GetCompressedObjectAsync(ns, id, HttpContext.RequestServices, supportsRedirectUri: supportsRedirect);
 
 				StringValues acceptHeader = Request.Headers["Accept"];
 				if (!acceptHeader.Contains("*/*") && acceptHeader.Count != 0 && !acceptHeader.Contains(mediaType))
diff --git a/Engine/Source/Programs/UnrealCloudDDC/Jupiter/Controllers/SymbolsController.cs b/Engine/Source/Programs/UnrealCloudDDC/Jupiter/Controllers/SymbolsController.cs
--- a/Engine/Source/Programs/UnrealCloudDDC/Jupiter/Controllers/SymbolsController.cs
+++ b/Engine/Source/Programs/UnrealCloudDDC/Jupiter/Controllers/SymbolsController.cs
@@ -71,12 +71,14 @@ namespace Jupiter.Controllers
 		/// <param name="moduleName">The name of the module, so Foo.pdb for instance</param>
 		/// <param name="identifier">The pdb identifier and age</param>
 		/// <param name="fileName">The specific file to fetch, either pdb or ptrs</param>
+		/// <param name="supportsRedirect">Set to true to indicate that you support getting redirected for this data</param>
 		[HttpGet("{ns}/{moduleName}/{identifier}/{fileName}", Order = 500)]
 		public async Task<IActionResult> GetAsync(
 			[FromRoute][Required] NamespaceId ns,
 			[FromRoute][Required] string moduleName,
 			[FromRoute][Required] string identifier,
-			[FromRoute][Required] string fileName)
+			[FromRoute][Required] string fileName,
+			[FromQuery] bool supportsRedirect = false)
 		{
 			ActionResult? accessResult = await _requestHelper.HasAccessToNamespaceAsync(User, Request, ns, new[] { JupiterAclAction.ReadObject });
 			if (accessResult != null)
@@ -106,7 +108,7 @@ namespace Jupiter.Controllers
 			CbObject cb = new CbObject(blobMemory);
 			IoHash payloadHash = cb["pdbPayload"].AsBinaryAttachment().Hash;
 
-			BlobContents referencedBlobContents = await _blobStore.GetObjectAsync(ns, BlobId.FromIoHash(payloadHash), storageLayers: null, supportsRedirectUri: true, bucketHint: bucket);
+			BlobContents referencedBlobContents = await _blobStore.GetObjectAsync(ns, BlobId.FromIoHash(payloadHash), storageLayers: null, supportsRedirectUri: supportsRedirect, bucketHint: bucket);
 
 			if (referencedBlobContents.RedirectUri != null)
 			{
