Cloud DDC - Allow all origins to call into Jupiter in CORS - enabling people to use their own web apps to call the api. They will need to send auth tokens anyway so no need to prevent CORS.

diff --git a/Engine/Source/Programs/UnrealCloudDDC/Jupiter/Jupiter.Startup.cs b/Engine/Source/Programs/UnrealCloudDDC/Jupiter/Jupiter.Startup.cs
--- a/Engine/Source/Programs/UnrealCloudDDC/Jupiter/Jupiter.Startup.cs
+++ b/Engine/Source/Programs/UnrealCloudDDC/Jupiter/Jupiter.Startup.cs
@@ -220,6 +220,16 @@ namespace Jupiter
 
 			services.AddSingleton<BlobStoreConsistencyCheckService>();
 			services.AddSingleton(typeof(IPeerServiceDiscovery), ServiceDiscoveryFactory);
+
+			// allow CORS from any domain to call this api as you will need a valid token anyway
+			services.AddCors(options =>
+			{
+				options.AddDefaultPolicy(policy =>
+				{
+					policy.AllowAnyOrigin();
+				});
+			});
+
 		}
 
 		private IBlockStore BlockIndexFactory(IServiceProvider provider)
