Cloud DDC - Fixes needed to make sure we pass CORS tests and thus allows web pages to consume (part of) our api.

diff --git a/Engine/Source/Programs/UnrealCloudDDC/Jupiter.Common/BaseStartup.cs b/Engine/Source/Programs/UnrealCloudDDC/Jupiter.Common/BaseStartup.cs
--- a/Engine/Source/Programs/UnrealCloudDDC/Jupiter.Common/BaseStartup.cs
+++ b/Engine/Source/Programs/UnrealCloudDDC/Jupiter.Common/BaseStartup.cs
@@ -96,6 +96,15 @@ namespace Jupiter
 			services.AddOptions<NamespaceSettings>().Bind(Configuration.GetSection("Namespaces")).ValidateDataAnnotations();
 
 			services.AddSingleton(typeof(INamespacePolicyResolver), typeof(NamespacePolicyResolver));
+			
+			// allow CORS from any domain to call this api as you will need a valid token anyway
+			services.AddCors(options =>
+			{
+				options.AddDefaultPolicy(policy =>
+				{
+					policy.AllowAnyOrigin();
+				});
+			});
 
 			services.AddControllers()
 				.AddMvcOptions(options =>
@@ -504,6 +513,8 @@ namespace Jupiter
 
 			app.UseRouting();
 
+			app.UseCors();
+
 			app.UseAuthentication();
 			app.UseAuthorization();
 
diff --git a/Engine/Source/Programs/UnrealCloudDDC/Jupiter/Controllers/BuildsController.cs b/Engine/Source/Programs/UnrealCloudDDC/Jupiter/Controllers/BuildsController.cs
--- a/Engine/Source/Programs/UnrealCloudDDC/Jupiter/Controllers/BuildsController.cs
+++ b/Engine/Source/Programs/UnrealCloudDDC/Jupiter/Controllers/BuildsController.cs
@@ -1296,6 +1296,18 @@ namespace Jupiter.Controllers
 
 		#region Enumeration endpoints
 
+		// adds a empty options endpoint for endpoints we expose to a webpage to pass cors preflight
+		// this should be improved if we decide we are maintaining the admin pages
+		[HttpOptions("")]
+		[HttpOptions("{ns}")]
+		[HttpOptions("{ns}/{bucket}/search")]
+		[ProducesDefaultResponseType]
+		public async Task<IActionResult> EmptyOptionsAsync()
+		{
+			await Task.CompletedTask;
+			return NoContent();
+		}
+
 		[HttpGet("")]
 		[ProducesDefaultResponseType]
 		public async Task<IActionResult> GetNamespacesAsync()
diff --git a/Engine/Source/Programs/UnrealCloudDDC/Jupiter/Jupiter.Startup.cs b/Engine/Source/Programs/UnrealCloudDDC/Jupiter/Jupiter.Startup.cs
--- a/Engine/Source/Programs/UnrealCloudDDC/Jupiter/Jupiter.Startup.cs
+++ b/Engine/Source/Programs/UnrealCloudDDC/Jupiter/Jupiter.Startup.cs
@@ -220,16 +220,6 @@ namespace Jupiter
 
 			services.AddSingleton<BlobStoreConsistencyCheckService>();
 			services.AddSingleton(typeof(IPeerServiceDiscovery), ServiceDiscoveryFactory);
-
-			// allow CORS from any domain to call this api as you will need a valid token anyway
-			services.AddCors(options =>
-			{
-				options.AddDefaultPolicy(policy =>
-				{
-					policy.AllowAnyOrigin();
-				});
-			});
-
 		}
 
 		private IBlockStore BlockIndexFactory(IServiceProvider provider)
