Cloud DDC - Allow any methods and any headers in the context of CORS, no need to be specific about this.

diff --git a/Engine/Source/Programs/UnrealCloudDDC/Jupiter.Common/BaseStartup.cs b/Engine/Source/Programs/UnrealCloudDDC/Jupiter.Common/BaseStartup.cs
--- a/Engine/Source/Programs/UnrealCloudDDC/Jupiter.Common/BaseStartup.cs
+++ b/Engine/Source/Programs/UnrealCloudDDC/Jupiter.Common/BaseStartup.cs
@@ -103,6 +103,8 @@ namespace Jupiter
 				options.AddDefaultPolicy(policy =>
 				{
 					policy.AllowAnyOrigin();
+					policy.AllowAnyHeader();
+					policy.AllowAnyMethod();
 				});
 			});
 
