Cloud DDC - Fixed issue were blob index were not updated during a multipart upload

diff --git a/Engine/Source/Programs/UnrealCloudDDC/Jupiter/Implementation/Blob/BlobService.cs b/Engine/Source/Programs/UnrealCloudDDC/Jupiter/Implementation/Blob/BlobService.cs
--- a/Engine/Source/Programs/UnrealCloudDDC/Jupiter/Implementation/Blob/BlobService.cs
+++ b/Engine/Source/Programs/UnrealCloudDDC/Jupiter/Implementation/Blob/BlobService.cs
@@ -790,13 +790,15 @@ public class BlobService : IBlobService
 		if (isCompressed)
 		{
 			(cid, blobId) = await this.PutCompressedObjectMetadataAsync(ns, bufferedPayload, ContentId.FromBlobIdentifier(id), _provider, cancellationToken: cancellationToken);
-
+			await _blobIndex.AddBlobToIndexAsync(ns, blobId, cancellationToken: cancellationToken);
 		}
 		else
 		{
 			ContentHash hash = await VerifyContentMatchesHashAsync(blobContents.Stream, id, cancellationToken: cancellationToken);
 			blobId = BlobId.FromContentHash(hash);
 			cid = ContentId.FromContentHash(hash);
+
+			await _blobIndex.AddBlobToIndexAsync(ns, blobId, cancellationToken: cancellationToken);
 		}
 
 		await blobStore.RenameMultipartBlobAsync(ns, blobName, blobId);
