Cloud DDC - Added flag to include TTL field for builds when searching - adds extra db requests so not enabled by default.

diff --git a/Engine/Source/Programs/UnrealCloudDDC/Jupiter/Controllers/BuildsController.cs b/Engine/Source/Programs/UnrealCloudDDC/Jupiter/Controllers/BuildsController.cs
--- a/Engine/Source/Programs/UnrealCloudDDC/Jupiter/Controllers/BuildsController.cs
+++ b/Engine/Source/Programs/UnrealCloudDDC/Jupiter/Controllers/BuildsController.cs
@@ -317,8 +317,9 @@ namespace Jupiter.Controllers
 			bool checkForMaxTime = searchRequest.Options.MaxTimeMs != 0u;
 			DateTime startTime = DateTime.Now;
 			TimeSpan maxTimeToSpend = TimeSpan.FromMilliseconds(searchRequest.Options.MaxTimeMs);
+			bool includeTTL = searchRequest.Options.IncludeTTL;
 			bool partialResult = false;
-			await foreach (BuildMetadata build in _buildStore.ListBuildsAsync(ns, bucket))
+			await foreach (BuildMetadata build in _buildStore.ListBuildsAsync(ns, bucket, includeTTL))
 			{
 				if (!build.IsFinalized)
 				{
@@ -1470,6 +1471,9 @@ namespace Jupiter.Controllers
 	{
 		[CbField("buildId")]
 		public CbObjectId BuildId { get; set; }
+
+		[CbField("ttl")]
+		public uint? TTL { get; set; }
 		[CbField("metadata")]
 		public CbObject Metadata { get; set; } = null!;
 	}
@@ -1493,7 +1497,7 @@ namespace Jupiter.Controllers
 		public SearchResult(List<BuildMetadata> foundResults, bool partialResult)
 		{
 			PartialResult = partialResult;
-			Results = foundResults.Select(b => new BuildSearchResult {BuildId = b.BuildId, Metadata = b.Metadata}).ToList();
+			Results = foundResults.Select(b => new BuildSearchResult {BuildId = b.BuildId, Metadata = b.Metadata, TTL = b.Ttl}).ToList();
 		}
 	}
 
@@ -1510,6 +1514,9 @@ namespace Jupiter.Controllers
 
 		[CbField("maxTimeMS")]
 		public uint MaxTimeMs { get; set; } = Debugger.IsAttached ? 0u : 120_000; // default to 2 minutes max per search unless debugging in which case we disable it
+
+		[CbField("includeTTL")]
+		public bool IncludeTTL { get; set; } = false;
 	}
 
 	public class SearchRequest
diff --git a/Engine/Source/Programs/UnrealCloudDDC/Jupiter/Implementation/Builds/IBuildStore.cs b/Engine/Source/Programs/UnrealCloudDDC/Jupiter/Implementation/Builds/IBuildStore.cs
--- a/Engine/Source/Programs/UnrealCloudDDC/Jupiter/Implementation/Builds/IBuildStore.cs
+++ b/Engine/Source/Programs/UnrealCloudDDC/Jupiter/Implementation/Builds/IBuildStore.cs
@@ -12,7 +12,7 @@ public interface IBuildStore
 	Task PutBuildAsync(NamespaceId ns, BucketId bucket, CbObjectId buildId, CbObject buildObject, uint ttl);
 	Task FinalizeBuildAsync(NamespaceId ns, BucketId bucket, CbObjectId buildId);
 	Task<BuildRecord?> GetBuildAsync(NamespaceId ns, BucketId bucket, CbObjectId buildId);
-	IAsyncEnumerable<BuildMetadata> ListBuildsAsync(NamespaceId ns, BucketId bucket);
+	IAsyncEnumerable<BuildMetadata> ListBuildsAsync(NamespaceId ns, BucketId bucket, bool includeTTL = false);
 
 	Task PutBuildPartAsync(NamespaceId ns, BucketId bucket, CbObjectId buildId, CbObjectId partId, string partName, uint ttl);
 	IAsyncEnumerable<(string, CbObjectId)> GetBuildPartsAsync(NamespaceId ns, BucketId bucket, CbObjectId buildId);
@@ -23,16 +23,18 @@ public interface IBuildStore
 
 public class BuildMetadata
 {
-	public BuildMetadata(CbObjectId buildId, CbObject metadata, bool isFinalized)
+	public BuildMetadata(CbObjectId buildId, CbObject metadata, bool isFinalized, uint? ttl)
 	{
 		BuildId = buildId;
 		Metadata = metadata;
 		IsFinalized = isFinalized;
+		Ttl = ttl;
 	}
 
 	public CbObjectId BuildId { get; set; }
 	public CbObject Metadata { get; set; } = null!;
 	public bool IsFinalized { get; set; }
+	public uint? Ttl { get; set; } = null;
 }
 
 public class BuildRecord
diff --git a/Engine/Source/Programs/UnrealCloudDDC/Jupiter/Implementation/Builds/MemoryBuildStore.cs b/Engine/Source/Programs/UnrealCloudDDC/Jupiter/Implementation/Builds/MemoryBuildStore.cs
--- a/Engine/Source/Programs/UnrealCloudDDC/Jupiter/Implementation/Builds/MemoryBuildStore.cs
+++ b/Engine/Source/Programs/UnrealCloudDDC/Jupiter/Implementation/Builds/MemoryBuildStore.cs
@@ -57,7 +57,7 @@ public class MemoryBuildStore : IBuildStore
 		return Task.CompletedTask;
 	}
 
-	public async IAsyncEnumerable<BuildMetadata> ListBuildsAsync(NamespaceId ns, BucketId bucket)
+	public async IAsyncEnumerable<BuildMetadata> ListBuildsAsync(NamespaceId ns, BucketId bucket, bool includeTTL = false)
 	{
 		await Task.CompletedTask;
 		foreach (MemoryBuildRecord build in _builds.Values.OrderByDescending(record => record.CreationDate))
@@ -67,7 +67,7 @@ public class MemoryBuildStore : IBuildStore
 				continue;
 			}
 
-			yield return new BuildMetadata(build.BuildId, build.BuildObject, build.IsFinalized);
+			yield return new BuildMetadata(build.BuildId, build.BuildObject, build.IsFinalized, includeTTL ? (uint)(build.ExpireTime - DateTime.Now).TotalSeconds : null);
 		}
 	}
 
diff --git a/Engine/Source/Programs/UnrealCloudDDC/Jupiter/Implementation/Builds/ScyllaBuildStore.cs b/Engine/Source/Programs/UnrealCloudDDC/Jupiter/Implementation/Builds/ScyllaBuildStore.cs
--- a/Engine/Source/Programs/UnrealCloudDDC/Jupiter/Implementation/Builds/ScyllaBuildStore.cs
+++ b/Engine/Source/Programs/UnrealCloudDDC/Jupiter/Implementation/Builds/ScyllaBuildStore.cs
@@ -51,7 +51,7 @@ public class ScyllaBuildStore : IBuildStore
 		_getTtlStatement = _session.Prepare("SELECT TTL(build_object) as ttl FROM builds WHERE namespace = ? AND bucket_id = ? AND build_id = ?");
 	}
 
-	public async IAsyncEnumerable<BuildMetadata> ListBuildsAsync(NamespaceId ns, BucketId bucket)
+	public async IAsyncEnumerable<BuildMetadata> ListBuildsAsync(NamespaceId ns, BucketId bucket, bool includeTTL = false)
 	{
 		using TelemetrySpan scope = _tracer.BuildScyllaSpan("scylla.list_builds");
 		foreach (ScyllaBuildEntry buildEntry in await _mapper.FetchAsync<ScyllaBuildEntry>("WHERE namespace = ? AND bucket_id = ?", ns.ToString(), bucket.ToString()))
@@ -62,6 +62,13 @@ public class ScyllaBuildStore : IBuildStore
 			{
 				continue;
 			}
+
+			if (includeTTL)
+			{
+				uint? ttl = await GetTTL(ns, bucket, metadata.BuildId);
+				metadata.Ttl = ttl;
+			}
+
 			yield return metadata;
 		}
 	}
@@ -192,7 +199,7 @@ public class ScyllaBuildEntry
 	{
 		try
 		{
-			return new BuildMetadata(new CbObjectId(BuildId), new CbObject(BuildObject), IsFinalized);
+			return new BuildMetadata(new CbObjectId(BuildId), new CbObject(BuildObject), IsFinalized, null);
 		}
 		catch (IndexOutOfRangeException)
 		{
