Cloud DDC - Updated publishing script for bootstrapper script to use the correct region for were the bucket is.

diff --git a/Engine/Source/Programs/UnrealCloudDDC/Bootstrapping/AWS/cmdlets.ps1 b/Engine/Source/Programs/UnrealCloudDDC/Bootstrapping/AWS/cmdlets.ps1
--- a/Engine/Source/Programs/UnrealCloudDDC/Bootstrapping/AWS/cmdlets.ps1
+++ b/Engine/Source/Programs/UnrealCloudDDC/Bootstrapping/AWS/cmdlets.ps1
@@ -7,7 +7,7 @@ function Upload-AWSBootstrapperToS3Dev {
     Get-ChildItem -Path $PSScriptRoot -Exclude *.ps1,*.zip |  Compress-Archive -CompressionLevel Fastest -DestinationPath setup.zip -Force
 
     Write-Host "Uploading to S3"
-    aws s3 cp setup.zip s3://devtools-misc/dev-cloud-ddc-bootstrapper/aws/setup.zip --acl public-read
+    aws s3 cp setup.zip s3://devtools-misc/dev-cloud-ddc-bootstrapper/aws/setup.zip --acl public-read --region eu-north-1
 }
 
 function Upload-AWSBootstrapperToS3Prod {
@@ -18,5 +18,5 @@ function Upload-AWSBootstrapperToS3Prod {
     Get-ChildItem -Path $PSScriptRoot -Exclude *.ps1,*.zip |  Compress-Archive -CompressionLevel Fastest -DestinationPath setup.zip -Force
 
     Write-Host "Uploading to S3"
-    aws s3 cp setup.zip s3://devtools-misc/cloud-ddc-bootstrapper/aws/setup.zip --acl public-read
+    aws s3 cp setup.zip s3://devtools-misc/cloud-ddc-bootstrapper/aws/setup.zip --acl public-read --region eu-north-1
 }
