Cloud DDC - Added option to helm charts to install the experimental dashboard. Currently not available to licensees and thus disabled by default.

diff --git a/Engine/Source/Programs/UnrealCloudDDC/Helm/UnrealCloudDDC/templates/configmap-dashboard.yaml b/Engine/Source/Programs/UnrealCloudDDC/Helm/UnrealCloudDDC/templates/configmap-dashboard.yaml
new file mode 100644
--- /dev/null
+++ b/Engine/Source/Programs/UnrealCloudDDC/Helm/UnrealCloudDDC/templates/configmap-dashboard.yaml
@@ -0,0 +1,20 @@
+{{- if .Values.dashboard.enabled -}}
+{{- $config := include "dashboard-config" . | fromYaml -}}
+apiVersion: v1
+kind: ConfigMap
+metadata:
+  name: {{ include "common.names.fullname" . }}-dashboard
+  labels: {{- include "common.labels.standard" . | nindent 4 }}
+data:
+  config.json: |
+{{- $config | toJson | nindent 4}} 
+{{- end -}}
+
+{{- define "dashboard-config" -}}
+{{- $config := .Values.dashboard.config | default dict -}}
+UpstreamUrl: {{ required "Upstream Url must be specified" $config.UpstreamUrl }}
+
+OidcAuthority: {{ default "" $config.OidcAuthority }}
+OidcClientId: {{ default "" $config.OidcClientId }}
+{{ end }}
+  
\ No newline at end of file
diff --git a/Engine/Source/Programs/UnrealCloudDDC/Helm/UnrealCloudDDC/templates/configmap.yaml b/Engine/Source/Programs/UnrealCloudDDC/Helm/UnrealCloudDDC/templates/configmap.yaml
--- a/Engine/Source/Programs/UnrealCloudDDC/Helm/UnrealCloudDDC/templates/configmap.yaml
+++ b/Engine/Source/Programs/UnrealCloudDDC/Helm/UnrealCloudDDC/templates/configmap.yaml
@@ -1,3 +1,4 @@
+{{- if .Values.enabled -}}
 {{- $baseSettings := include "jupiter.applicationSettings" . | fromYaml -}}
 {{- $overrides := include "storage.settings" . | fromYaml -}}
 {{- $newDict := merge $overrides $baseSettings -}}
@@ -9,6 +10,7 @@ metadata:
 data:
   appsettings.Local.yaml: |
 {{ $newDict | toYaml | indent 4 }} 
+{{ end }}
 
 {{- define "storage.settings" -}}
 {{- $serviceBaseName := .Release.Name | lower -}}
@@ -31,7 +33,7 @@ UnrealCloudDDC:
   BuildStoreImplementation :  {{ default "Scylla" $config.BuildStoreImplementation }}
 
 {{- if hasKey $config "EnableOnDemandReplication" }}
-  EnableOnDemandReplication :  {{ $config.EnableOnDemandReplication }}
+  EnableOnDemandReplication :  {{ $config.EnableOnDemandReplication }}  
 {{- end }}
 
 {{- if hasKey $config "EnableLastAccessTracking" }}
@@ -221,4 +223,3 @@ S3:
   {{- end }}
   
 {{ end }}
-
diff --git a/Engine/Source/Programs/UnrealCloudDDC/Helm/UnrealCloudDDC/templates/deployment-dashboard.yaml b/Engine/Source/Programs/UnrealCloudDDC/Helm/UnrealCloudDDC/templates/deployment-dashboard.yaml
new file mode 100644
--- /dev/null
+++ b/Engine/Source/Programs/UnrealCloudDDC/Helm/UnrealCloudDDC/templates/deployment-dashboard.yaml
@@ -0,0 +1,107 @@
+{{- if .Values.dashboard.enabled -}}
+{{ $fullname := include "common.names.fullname" . }}
+{{ $dashboard := .Values.dashboard }}
+apiVersion: apps/v1
+kind: Deployment
+metadata:
+  name: {{ $fullname }}
+  labels: {{- include "common.labels.standard" . | nindent 4 }}
+spec:
+  selector:
+    matchLabels: {{- include "common.labels.matchLabels" . | nindent 6 }}
+  replicas: {{ $dashboard.replicaCount }}
+  template:
+    metadata:
+      labels: {{- include "common.labels.standard" . | nindent 8 }}
+      {{- if $dashboard.podLabels }}
+      {{- include "common.tplvalues.render" ( dict "value" $dashboard.podLabels "context" $ ) | nindent 8 }}
+      {{- end }}
+
+      {{- if or $dashboard.podAnnotations $dashboard.restartPodOnConfigMapChange }}
+      annotations:
+        {{- if $dashboard.podAnnotations }}
+        {{- include "common.tplvalues.render" (dict "value" $dashboard.podAnnotations "context" $) | nindent 8 }}
+        {{- end }}
+        {{- if $dashboard.restartPodOnConfigMapChange }}
+        checksum/config: {{ include (print $.Template.BasePath "/configmap.yaml") . | sha256sum }}
+        {{- end }}
+      {{- end }}
+    spec:
+      serviceAccountName: {{ include "jupiter.serviceAccountName" . }}
+      {{- with $dashboard.imagePullSecrets }}
+      imagePullSecrets:
+        {{- toYaml . | nindent 8 }}
+      {{- end }}
+      {{- if $dashboard.terminationGracePeriodSeconds }}
+      terminationGracePeriodSeconds: {{ $dashboard.terminationGracePeriodSeconds }}
+      {{- end }}
+      {{- if $dashboard.affinity }}
+      affinity: {{- include "common.tplvalues.render" (dict "value" $dashboard.affinity "context" $) | nindent 8 }}
+      {{- else }}
+      affinity:
+        podAffinity: {{- include "common.affinities.pods" (dict "type" $dashboard.podAffinityPreset "context" $) | nindent 10 }}
+        podAntiAffinity: {{- include "common.affinities.pods" (dict "type" $dashboard.podAntiAffinityPreset "context" $) | nindent 10 }}
+        nodeAffinity: {{- include "common.affinities.nodes" (dict "type" $dashboard.nodeAffinityPreset.type "key" $dashboard.nodeAffinityPreset.key "values" $dashboard.nodeAffinityPreset.values) | nindent 10 }}
+      {{- end }}
+      {{- if $dashboard.securityContext.enabled }}
+      {{ $securityContext := unset $dashboard.securityContext "enabled" }}
+      securityContext: {{- include "common.tplvalues.render" (dict "value" $securityContext "context" $) | nindent 8 }}
+      {{- end }}
+      {{- if $dashboard.nodeSelector }}
+      nodeSelector: {{- include "common.tplvalues.render" (dict "value" $dashboard.nodeSelector "context" $) | nindent 8 }}
+      {{- end }}
+      {{- if $dashboard.tolerations }}
+      tolerations: {{- include "common.tplvalues.render" (dict "value" $dashboard.tolerations "context" $) | nindent 8 }}
+      {{- end }}
+      containers:
+      - name: {{ include "common.names.name" . }}-dashboard
+        image: {{ printf "%s:%s" $dashboard.image.repository (include "jupiter.appversion" . ) }}
+        imagePullPolicy: {{ $dashboard.image.pullPolicy }}
+        env:
+        - name: Kubernetes__Namespace
+          valueFrom:
+            fieldRef:
+              fieldPath: metadata.namespace
+        - name: Kubernetes__PodName
+          valueFrom:
+            fieldRef:
+              fieldPath: metadata.name
+        {{- if $dashboard.env }}
+        {{- include "common.tplvalues.render" ( dict "value" $dashboard.env "context" $ ) | nindent 8 }}
+        {{- end }}
+        ports:
+        - name: http
+          containerPort: 80
+        {{- if $dashboard.livenessProbe.enabled }}
+        livenessProbe:
+          httpGet:
+            path: {{ default "/index.html" $dashboard.livenessProbe.path }}
+            port: http
+          initialDelaySeconds: {{ $dashboard.livenessProbe.initialDelaySeconds }}
+          periodSeconds: {{ $dashboard.livenessProbe.periodSeconds }}
+          timeoutSeconds: {{ $dashboard.livenessProbe.timeoutSeconds }}
+          successThreshold: {{ $dashboard.livenessProbe.successThreshold }}
+          failureThreshold: {{ $dashboard.livenessProbe.failureThreshold }}
+        {{- end }}
+        {{- if $dashboard.readinessProbe.enabled }}
+        readinessProbe:
+          httpGet:
+            path:  {{ default "/index.html" $dashboard.readinessProbe.path }}
+            port: http
+          initialDelaySeconds: {{ $dashboard.readinessProbe.initialDelaySeconds }}
+          periodSeconds: {{ $dashboard.readinessProbe.periodSeconds }}
+          timeoutSeconds: {{ $dashboard.readinessProbe.timeoutSeconds }}
+          successThreshold: {{ $dashboard.readinessProbe.successThreshold }}
+          failureThreshold: {{ $dashboard.readinessProbe.failureThreshold }}
+        {{- end }}
+        {{- if $dashboard.resources }}
+        resources: {{- toYaml $dashboard.resources | nindent 12 }}
+        {{- end }}
+        volumeMounts:
+        - name: config-volume
+          mountPath: /config
+      volumes:
+      - name: config-volume
+        configMap:
+          name: {{ include "common.names.fullname" . }}-dashboard
+{{- end -}}
\ No newline at end of file
diff --git a/Engine/Source/Programs/UnrealCloudDDC/Helm/UnrealCloudDDC/templates/deployment.yaml b/Engine/Source/Programs/UnrealCloudDDC/Helm/UnrealCloudDDC/templates/deployment.yaml
--- a/Engine/Source/Programs/UnrealCloudDDC/Helm/UnrealCloudDDC/templates/deployment.yaml
+++ b/Engine/Source/Programs/UnrealCloudDDC/Helm/UnrealCloudDDC/templates/deployment.yaml
@@ -1,3 +1,4 @@
+{{- if .Values.enabled -}}
 {{ $fullname := include "common.names.fullname" . }}
 apiVersion: apps/v1
 kind: Deployment
@@ -276,4 +277,5 @@ spec:
       {{- end }}
       {{- if .Values.extraVolumes }}
       {{- include "common.tplvalues.render" ( dict "value" .Values.extraVolumes "context" $) | nindent 6 }}
-      {{- end }}
\ No newline at end of file
+      {{- end }}
+{{- end -}}
\ No newline at end of file
diff --git a/Engine/Source/Programs/UnrealCloudDDC/Helm/UnrealCloudDDC/templates/ingress-dashboard.yaml b/Engine/Source/Programs/UnrealCloudDDC/Helm/UnrealCloudDDC/templates/ingress-dashboard.yaml
new file mode 100644
--- /dev/null
+++ b/Engine/Source/Programs/UnrealCloudDDC/Helm/UnrealCloudDDC/templates/ingress-dashboard.yaml
@@ -0,0 +1,52 @@
+{{- if .Values.dashboard.ingress.enabled }}
+{{- $ingress := .Values.dashboard.ingress | default dict -}}
+{{- $servicePort := default "http" .Values.dashbord.ingress.servicePort -}}
+apiVersion: {{ include "common.capabilities.ingress.apiVersion" . }}
+kind: Ingress
+metadata:
+  name: {{ include "common.names.fullname" . }}
+  labels: {{- include "common.labels.standard" . | nindent 4 }}
+  annotations:
+    {{- range $key, $value := $ingress.annotations }}
+    {{ $key }}: {{ $value | quote }}
+    {{- end }}
+spec:
+  {{- if and $ingress.ingressClassName (include "common.ingress.supportsIngressClassname" .) }}
+  ingressClassName: {{ $ingress.ingressClassName | quote }}
+  {{- end }}
+  rules:
+    {{- if $ingress.hostname }}
+    - host: {{ $ingress.hostname | quote }}
+      http:
+        paths:
+          {{- if $ingress.extraPaths }}
+          {{- toYaml $ingress.extraPaths | nindent 10 }}
+          {{- end }}
+          - path: {{ $ingress.path }}
+            {{- if eq "true" (include "common.ingress.supportsPathType" .) }}
+            pathType: {{ $ingress.pathType }}
+            {{- end }}
+            backend: {{- include "common.ingress.backend" (dict "serviceName" (include "common.names.fullname" .) "servicePort" $servicePort "context" $)  | nindent 14 }}
+    {{- end }}
+    {{- range $ingress.extraHosts }}
+    - host: {{ .name }}
+      http:
+        paths:
+          - path: {{ default "/" .path }}
+            {{- if eq "true" (include "common.ingress.supportsPathType" $) }}
+            pathType: {{ default "ImplementationSpecific" .pathType }}
+            {{- end }}
+            backend: {{- include "common.ingress.backend" (dict "serviceName" (include "common.names.fullname" $) "servicePort" $servicePort "context" $) | nindent 14 }}
+    {{- end }}
+  {{- if or $ingress.tls .Values.ingress.extraTls .Values.ingress.hosts }}
+  tls:
+    {{- if $ingress.tls }}
+    - hosts:
+        - {{ $ingress.hostname }}
+      secretName: {{ printf "%s-tls" $ingress.hostname | trunc 63 | trimSuffix "-" }}
+    {{- end }}
+    {{- if $ingress.extraTls }}
+    {{- toYaml $ingress.extraTls | nindent 4 }}
+    {{- end }}
+  {{- end }}
+{{- end }}
\ No newline at end of file
diff --git a/Engine/Source/Programs/UnrealCloudDDC/Helm/UnrealCloudDDC/templates/service.yaml b/Engine/Source/Programs/UnrealCloudDDC/Helm/UnrealCloudDDC/templates/service.yaml
--- a/Engine/Source/Programs/UnrealCloudDDC/Helm/UnrealCloudDDC/templates/service.yaml
+++ b/Engine/Source/Programs/UnrealCloudDDC/Helm/UnrealCloudDDC/templates/service.yaml
@@ -1,3 +1,4 @@
+{{- if .Values.enabled -}}
 apiVersion: v1
 kind: Service
 metadata:
@@ -25,4 +26,6 @@ spec:
   {{- toYaml .Values.service.extraPort | nindent 4 }}
   {{- end }}
 
-  selector: {{- include "common.labels.matchLabels" . | nindent 4 }}
\ No newline at end of file
+  selector: {{- include "common.labels.matchLabels" . | nindent 4 }}
+
+{{ end }}
diff --git a/Engine/Source/Programs/UnrealCloudDDC/Helm/UnrealCloudDDC/values.yaml b/Engine/Source/Programs/UnrealCloudDDC/Helm/UnrealCloudDDC/values.yaml
--- a/Engine/Source/Programs/UnrealCloudDDC/Helm/UnrealCloudDDC/values.yaml
+++ b/Engine/Source/Programs/UnrealCloudDDC/Helm/UnrealCloudDDC/values.yaml
@@ -20,6 +20,10 @@ image:
 ##
 kubeVersion:
 
+## Enable the main workload (cloud ddc) - mostly used if you want to break the install into multiple parts
+## 
+enabled: true
+
 ## String to partially override common.names.fullname template (will maintain the release name)
 ##
 # nameOverride:
@@ -383,6 +387,77 @@ nginx:
     successThreshold: 1
     failureThreshold: 5
 
+## Dashboard configuration
+##
+dashboard:
+  ## Set to enable the admin dashboard webpage
+  ##
+  enabled: False
+
+  ingress:
+    ## Set to enable creating a ingress (load balancer)
+    enabled: False
+  ## Liveness and readiness probes, defaults to corresponding endpoints
+  ##
+  livenessProbe:
+    enabled: true
+    initialDelaySeconds: 5
+    periodSeconds: 10
+    timeoutSeconds: 5
+    successThreshold: 1
+    failureThreshold: 5
+  readinessProbe:
+    enabled: true
+    initialDelaySeconds: 5
+    periodSeconds: 10
+    timeoutSeconds: 3
+    successThreshold: 1
+    failureThreshold: 5
+  
+  ## Security Context
+  ## ref: https://kubernetes.io/docs/tasks/configure-pod-container/security-context/
+  ##
+  securityContext: {}
+
+  ## Pod annotations
+  ## ref: https://kubernetes.io/docs/concepts/overview/working-with-objects/annotations/
+  ##
+  podAnnotations: {}
+
+  ## Pod affinity preset
+  ## ref: https://kubernetes.io/docs/concepts/scheduling-eviction/assign-pod-node/#inter-pod-affinity-and-anti-affinity
+  ## Allowed values: soft, hard
+  ##
+  podAffinityPreset: ""
+
+  ## Pod anti-affinity preset
+  ## ref: https://kubernetes.io/docs/concepts/scheduling-eviction/assign-pod-node/#inter-pod-affinity-and-anti-affinity
+  ## Allowed values: soft, hard
+  ##
+  podAntiAffinityPreset: soft
+
+  ## Node affinity preset
+  ## ref: https://kubernetes.io/docs/concepts/scheduling-eviction/assign-pod-node/#node-affinity
+  ## Allowed values: soft, hard
+  ##
+  nodeAffinityPreset:
+    ## Node affinity type
+    ## Allowed values: soft, hard
+    ##
+    type: ""
+    ## Node label key to match
+    ## E.g.
+    ## key: "kubernetes.io/e2e-az-name"
+    ##
+    key: ""
+    ## Node label values to match
+    ## E.g.
+    ## values:
+    ##   - e2e-az1
+    ##   - e2e-az2
+    ##
+    values: []
+
 ## Worker configuration
 worker:
 
