Cloud DDC - Added audit logs for the build endpoints, disabled by default.

diff --git a/Engine/Source/Programs/UnrealCloudDDC/Helm/UnrealCloudDDC/templates/configmap.yaml b/Engine/Source/Programs/UnrealCloudDDC/Helm/UnrealCloudDDC/templates/configmap.yaml
--- a/Engine/Source/Programs/UnrealCloudDDC/Helm/UnrealCloudDDC/templates/configmap.yaml
+++ b/Engine/Source/Programs/UnrealCloudDDC/Helm/UnrealCloudDDC/templates/configmap.yaml
@@ -177,6 +177,11 @@ Symbols:
 {{ $config.Symbols | toYaml | nindent 2 }}
 {{- end }}
 
+{{- if $config.Builds }}
+Builds:
+{{ $config.Builds | toYaml | nindent 2 }}
+{{- end }}
+
 {{- if and (.Values.global.cloudProvider) (eq .Values.global.cloudProvider "Azure") }}
 {{ include "cloud.settings.azure" .Values }}
 {{- else if and (.Values.global.cloudProvider) (eq .Values.global.cloudProvider "AWS") }}
diff --git a/Engine/Source/Programs/UnrealCloudDDC/Jupiter/Controllers/BuildsController.cs b/Engine/Source/Programs/UnrealCloudDDC/Jupiter/Controllers/BuildsController.cs
--- a/Engine/Source/Programs/UnrealCloudDDC/Jupiter/Controllers/BuildsController.cs
+++ b/Engine/Source/Programs/UnrealCloudDDC/Jupiter/Controllers/BuildsController.cs
@@ -30,13 +30,16 @@ using System;
 using System.Collections.Concurrent;
 using System.Diagnostics;
 using System.Diagnostics.Metrics;
+using System.Net.Http;
 using System.Text.Json;
 using System.Text.Json.Serialization;
 using System.Threading;
 using Jupiter.Implementation.Blob;
 using Jupiter.Implementation.Builds;
+using Microsoft.Extensions.Options;
 using ContentHash = Jupiter.Implementation.ContentHash;
 using CbField = EpicGames.Serialization.CbField;
+using ILogger = Serilog.ILogger;
 using NotImplementedException = System.NotImplementedException;
 
 namespace Jupiter.Controllers
@@ -63,8 +66,9 @@ namespace Jupiter.Controllers
 		private readonly IBuildStore _buildStore;
 		private readonly IBlobIndex _blobIndex;
 		private readonly IReferenceResolver _referenceResolver;
+		private readonly ILogger? _auditLogger;
 
-		public BuildsController(IRefService refService, IBlobService blobStore, IBlockStore blockStore, IBuildStore buildStore, IBlobIndex blobIndex, IReferenceResolver referenceResolver, IDiagnosticContext diagnosticContext, FormatResolver formatResolver, BufferedPayloadFactory bufferedPayloadFactory, NginxRedirectHelper nginxRedirectHelper, INamespacePolicyResolver namespacePolicyResolver, IRequestHelper requestHelper, Tracer tracer, ILogger<BuildsController> logger)
+		public BuildsController(IRefService refService, IBlobService blobStore, IBlockStore blockStore, IBuildStore buildStore, IBlobIndex blobIndex, IReferenceResolver referenceResolver, IDiagnosticContext diagnosticContext, FormatResolver formatResolver, BufferedPayloadFactory bufferedPayloadFactory, NginxRedirectHelper nginxRedirectHelper, INamespacePolicyResolver namespacePolicyResolver, IRequestHelper requestHelper, Tracer tracer, ILogger<BuildsController> logger, IOptionsMonitor<BuildsSettings> settings)
 		{
 			_refService = refService;
 			_blobStore = blobStore;
@@ -80,6 +84,7 @@ namespace Jupiter.Controllers
 			_requestHelper = requestHelper;
 			_tracer = tracer;
 			_logger = logger;
+			_auditLogger = settings.CurrentValue.EnableAuditLog ? Serilog.Log.ForContext("LogType", "Audit") : null;
 		}
 
 		#region Build endpoints
@@ -97,6 +102,8 @@ namespace Jupiter.Controllers
 			Tracer.CurrentSpan.SetAttribute("bucket", bucket.ToString());
 			Tracer.CurrentSpan.SetAttribute("namespace", ns.ToString());
 
+			LogAuditEntry(HttpMethod.Put, nameof(PutBuildAsync), ns, bucket, buildId, null);
+
 			CbObject buildObject;
 			try
 			{
@@ -202,6 +209,8 @@ namespace Jupiter.Controllers
 			Tracer.CurrentSpan.SetAttribute("bucket", bucket.ToString());
 			Tracer.CurrentSpan.SetAttribute("namespace", ns.ToString());
 
+			LogAuditEntry(HttpMethod.Post, nameof(UpdateTTLAsync), ns, bucket, buildId, null);
+
 			await foreach ((string _, CbObjectId partId) in _buildStore.GetBuildPartsAsync(ns, bucket, buildId))
 			{
 				RefId refId = RefId.FromName($"{buildId}/{partId}");
@@ -226,6 +235,8 @@ namespace Jupiter.Controllers
 			Tracer.CurrentSpan.SetAttribute("bucket", bucket.ToString());
 			Tracer.CurrentSpan.SetAttribute("namespace", ns.ToString());
 
+			LogAuditEntry(HttpMethod.Delete, nameof(DeleteBuildAsync), ns, bucket, buildId, null);
+
 			await foreach ((string _, CbObjectId partId) in _buildStore.GetBuildPartsAsync(ns, bucket, buildId))
 			{
 				RefId refId = RefId.FromName($"{buildId}/{partId}");
@@ -250,6 +261,8 @@ namespace Jupiter.Controllers
 			Tracer.CurrentSpan.SetAttribute("bucket", bucket.ToString());
 			Tracer.CurrentSpan.SetAttribute("namespace", ns.ToString());
 
+			LogAuditEntry(HttpMethod.Get, nameof(GetBuildAsync), ns, bucket, buildId, null);
+
 			BuildRecord? buildRecord = await _buildStore.GetBuildAsync(ns, bucket, buildId);
 
 			if (buildRecord == null)
@@ -309,7 +322,7 @@ namespace Jupiter.Controllers
 			{
 				return accessResult;
 			}
-
+			
 			List<ISearchOp> searchOps = searchRequest.ToSearchOps();
 
 			List<BuildMetadata> foundResults = new List<BuildMetadata>();
@@ -387,6 +400,8 @@ namespace Jupiter.Controllers
 				return accessResult;
 			}
 
+			LogAuditEntry(HttpMethod.Get, nameof(GetBuildPartAsync), ns, bucket, buildId, partId);
+
 			try
 			{
 				Tracer.CurrentSpan.SetAttribute("bucket", bucket.ToString());
@@ -499,6 +514,8 @@ namespace Jupiter.Controllers
 			Tracer.CurrentSpan.SetAttribute("bucket", bucket.ToString());
 			Tracer.CurrentSpan.SetAttribute("namespace", ns.ToString());
 
+			LogAuditEntry(HttpMethod.Put, nameof(PutBuildPartAsync), ns, bucket, buildId, partName);
+
 			CbObject payloadObject;
 			BlobId payloadHash;
 			try
@@ -1330,6 +1347,7 @@ namespace Jupiter.Controllers
 			{
 				return new ForbidResult();
 			}
+			_auditLogger?.Information("{HttpMethod} {Endpoint} IP:{IP} User:{Username} UserAgent:\"{Useragent}\"", HttpMethod.Get, nameof(GetNamespacesAsync), Request.HttpContext.Connection.RemoteIpAddress, User?.Identity?.Name ?? "Unknown-user", string.Join(' ', Request.Headers.UserAgent.ToArray()));
 
 			return Ok(new GetNamespacesResponse(namespacesWithAccess.ToArray()));
 		}
@@ -1338,6 +1356,8 @@ namespace Jupiter.Controllers
 		[ProducesDefaultResponseType]
 		public async Task<IActionResult> GetBucketsAsync(NamespaceId ns)
 		{
+			_auditLogger?.Information("{HttpMethod} {Endpoint} {Namespace} IP:{IP} User:{Username} UserAgent:\"{Useragent}\"", HttpMethod.Get, nameof(GetBucketsAsync), ns, Request.HttpContext.Connection.RemoteIpAddress, User.Identity?.Name ?? "Unknown-user", string.Join(' ', Request.Headers.UserAgent.ToArray()));
+
 			BucketId[] potentialBuckets = await _refService.GetBucketsAsync(ns, HttpContext.RequestAborted).ToArrayAsync(HttpContext.RequestAborted);
 
 			// filter buckets down to only the ones the user has access to
@@ -1359,6 +1379,16 @@ namespace Jupiter.Controllers
 		}
 
 		#endregion
+
+		private void LogAuditEntry(HttpMethod method, string endpoint, NamespaceId ns, BucketId bucketId, CbObjectId buildId, string? partId)
+		{
+			_auditLogger?.Information("{HttpMethod} {Endpoint} '{Namespace}/{BucketId}:{BuildId}' Part:{PartId} IP:{IP} User:{Username} UserAgent:\"{Useragent}\"", method, endpoint, ns, bucketId, buildId, partId, Request.HttpContext.Connection.RemoteIpAddress, User?.Identity?.Name ?? "Unknown-user", string.Join(' ', Request.Headers.UserAgent.ToArray()));
+		}
+	}
+
+	public class BuildsSettings
+	{
+		public bool EnableAuditLog { get; set; }
 	}
 
 	public class GetBlockMetadataRequest
diff --git a/Engine/Source/Programs/UnrealCloudDDC/Jupiter/Jupiter.Startup.cs b/Engine/Source/Programs/UnrealCloudDDC/Jupiter/Jupiter.Startup.cs
--- a/Engine/Source/Programs/UnrealCloudDDC/Jupiter/Jupiter.Startup.cs
+++ b/Engine/Source/Programs/UnrealCloudDDC/Jupiter/Jupiter.Startup.cs
@@ -97,6 +97,7 @@ namespace Jupiter
 			services.AddOptions<NginxSettings>().Bind(Configuration.GetSection("Nginx")).ValidateDataAnnotations();
 
 			services.AddOptions<SymbolsSettings>().Bind(Configuration.GetSection("Symbols")).ValidateDataAnnotations();
+			services.AddOptions<BuildsSettings>().Bind(Configuration.GetSection("Builds")).ValidateDataAnnotations();
 
 			services.AddOptions<ConsistencyCheckSettings>().Bind(Configuration.GetSection("ConsistencyCheck")).ValidateDataAnnotations();
 			services.AddOptions<ReferenceResolverSettings>().Bind(Configuration.GetSection("ReferenceResolver")).ValidateDataAnnotations();
