Cloud DDC - Fixed potential issues were temp files were not being cleaned up properly for very large ref objects (mostly something that is encountered in builds as DDC tends to have small enough refs to fit into memory).

diff --git a/Engine/Source/Programs/UnrealCloudDDC/Jupiter/Controllers/BuildsController.cs b/Engine/Source/Programs/UnrealCloudDDC/Jupiter/Controllers/BuildsController.cs
--- a/Engine/Source/Programs/UnrealCloudDDC/Jupiter/Controllers/BuildsController.cs
+++ b/Engine/Source/Programs/UnrealCloudDDC/Jupiter/Controllers/BuildsController.cs
@@ -408,13 +408,15 @@ namespace Jupiter.Controllers
 				Tracer.CurrentSpan.SetAttribute("namespace", ns.ToString());
 
 				RefId refId = RefId.FromName($"{buildId}/{partId}");
-				(RefRecord objectRecord, BlobContents? blob) = await _refService.GetAsync(ns, bucket, refId, Array.Empty<string>());
+				(RefRecord objectRecord, BlobContents? maybeBlob) = await _refService.GetAsync(ns, bucket, refId, Array.Empty<string>());
 
-				if (blob == null)
+				if (maybeBlob == null)
 				{
 					throw new InvalidOperationException($"Blob was null when attempting to fetch {ns} {bucket} {buildId}");
 				}
 
+				await using BlobContents blob = maybeBlob;
+
 				if (!objectRecord.IsFinalized)
 				{
 					// we do not consider un-finalized objects as valid
@@ -798,6 +800,7 @@ namespace Jupiter.Controllers
 			{
 				(BlobContents blobContents, string mediaType) = await _blobStore.GetCompressedObjectAsync(ns, id, HttpContext.RequestServices, supportsRedirectUri: supportsRedirect);
 
+				await using BlobContents blob = blobContents;
 				StringValues acceptHeader = Request.Headers["Accept"];
 				if (!acceptHeader.Contains("*/*") && acceptHeader.Count != 0 && !acceptHeader.Contains(mediaType))
 				{
diff --git a/Engine/Source/Programs/UnrealCloudDDC/Jupiter/Controllers/ReferencesController.cs b/Engine/Source/Programs/UnrealCloudDDC/Jupiter/Controllers/ReferencesController.cs
--- a/Engine/Source/Programs/UnrealCloudDDC/Jupiter/Controllers/ReferencesController.cs
+++ b/Engine/Source/Programs/UnrealCloudDDC/Jupiter/Controllers/ReferencesController.cs
@@ -162,13 +162,15 @@ namespace Jupiter.Controllers
 				Tracer.CurrentSpan.SetAttribute("bucket", bucket.ToString());
 				Tracer.CurrentSpan.SetAttribute("namespace", ns.ToString());
 
-				(RefRecord objectRecord, BlobContents? blob) = await _refService.GetAsync(ns, bucket, key, Array.Empty<string>());
+				(RefRecord objectRecord, BlobContents? maybeBlob) = await _refService.GetAsync(ns, bucket, key, Array.Empty<string>());
 
-				if (blob == null)
+				if (maybeBlob == null)
 				{
 					throw new InvalidOperationException($"Blob was null when attempting to fetch {ns} {bucket} {key}");
 				}
 
+				await using BlobContents blob = maybeBlob;
+
 				if (!objectRecord.IsFinalized)
 				{
 					// we do not consider un-finalized objects as valid
diff --git a/Engine/Source/Programs/UnrealCloudDDC/Jupiter/Implementation/Blob/BlobService.cs b/Engine/Source/Programs/UnrealCloudDDC/Jupiter/Implementation/Blob/BlobService.cs
--- a/Engine/Source/Programs/UnrealCloudDDC/Jupiter/Implementation/Blob/BlobService.cs
+++ b/Engine/Source/Programs/UnrealCloudDDC/Jupiter/Implementation/Blob/BlobService.cs
@@ -784,7 +784,7 @@ public class BlobService : IBlobService
 			throw new Exception($"Failed to find multipart blob {blobName}");
 		}
 
-		IBufferedPayload bufferedPayload = await _bufferedPayloadFactory.CreateFromStreamAsync(blobContents.Stream, blobContents.Length, cancellationToken);
+		using IBufferedPayload bufferedPayload = await _bufferedPayloadFactory.CreateFromStreamAsync(blobContents.Stream, blobContents.Length, cancellationToken);
 		ContentId cid;
 		BlobId blobId;
 		if (isCompressed)
