Cloud DDC - Added chunk size to build metadata on upload as its hard to set this when uploading the build since its returned as a part of the response.

diff --git a/Engine/Source/Programs/UnrealCloudDDC/Jupiter/Controllers/BuildsController.cs b/Engine/Source/Programs/UnrealCloudDDC/Jupiter/Controllers/BuildsController.cs
--- a/Engine/Source/Programs/UnrealCloudDDC/Jupiter/Controllers/BuildsController.cs
+++ b/Engine/Source/Programs/UnrealCloudDDC/Jupiter/Controllers/BuildsController.cs
@@ -148,10 +148,22 @@ namespace Jupiter.Controllers
 				return BadRequest(e.Message);
 			}
 
+			// copy the cb object to a writer so that we can append useful metadata fields
+			CbWriter newObjectWriter = new CbWriter();
+			newObjectWriter.BeginObject();
+			foreach (CbField field in buildObject)
+			{
+				newObjectWriter.WriteField(field);
+			}
+			uint idealChunkSize = _blobStore.GetMultipartLimits(ns)?.IdealChunkSize ?? DefaultChunkSize;
+			newObjectWriter.WriteInteger("chunkSize", idealChunkSize);
+			newObjectWriter.EndObject();
+			CbObject mutatedObject = newObjectWriter.ToObject();
+
 			uint ttl = (uint)_namespacePolicyResolver.GetPoliciesForNs(ns).DefaultTTL.TotalSeconds;
-			await _buildStore.PutBuildAsync(ns, bucket, buildId, buildObject, ttl);
+			await _buildStore.PutBuildAsync(ns, bucket, buildId, mutatedObject, ttl);
 
-			return Ok(new PutBuildResponse {ChunkSize = _blobStore.GetMultipartLimits(ns)?.IdealChunkSize ?? DefaultChunkSize});
+			return Ok(new PutBuildResponse {ChunkSize = idealChunkSize});
 		}
 
 		private const int DefaultChunkSize = 32 * 1024 * 1024; // Default to 32 MB chunks
