Cloud DDC - Removed disposing of the BlobContents when fetching blob, fixes disposed access when FileStreamResult is processing but means we will leak any temporary file we allocated for this operation.

diff --git a/Engine/Source/Programs/UnrealCloudDDC/Jupiter/Controllers/BuildsController.cs b/Engine/Source/Programs/UnrealCloudDDC/Jupiter/Controllers/BuildsController.cs
--- a/Engine/Source/Programs/UnrealCloudDDC/Jupiter/Controllers/BuildsController.cs
+++ b/Engine/Source/Programs/UnrealCloudDDC/Jupiter/Controllers/BuildsController.cs
@@ -812,7 +812,6 @@ namespace Jupiter.Controllers
 			{
 				(BlobContents blobContents, string mediaType) = await _blobStore.GetCompressedObjectAsync(ns, id, HttpContext.RequestServices, supportsRedirectUri: supportsRedirect);
 
-				await using BlobContents blob = blobContents;
 				StringValues acceptHeader = Request.Headers["Accept"];
 				if (!acceptHeader.Contains("*/*") && acceptHeader.Count != 0 && !acceptHeader.Contains(mediaType))
 				{
