Cloud DDC - Introduced debug prefix on all temp files so that we can more easily determine were a temp file was allocated if its still kept around

diff --git a/Engine/Source/Programs/UnrealCloudDDC/Jupiter.Common/Implementation/StreamingBodyPayload.cs b/Engine/Source/Programs/UnrealCloudDDC/Jupiter.Common/Implementation/StreamingBodyPayload.cs
--- a/Engine/Source/Programs/UnrealCloudDDC/Jupiter.Common/Implementation/StreamingBodyPayload.cs
+++ b/Engine/Source/Programs/UnrealCloudDDC/Jupiter.Common/Implementation/StreamingBodyPayload.cs
@@ -111,9 +111,9 @@ namespace Jupiter.Common.Implementation
 
 		public FileInfo TempFile => _tempFile;
 
-		private FilesystemBufferedPayload(string filesystemRoot)
+		private FilesystemBufferedPayload(string filesystemRoot, string debugPrefix)
 		{
-			_tempFile = new FileInfo(Path.Combine(filesystemRoot, Path.GetRandomFileName()));
+			_tempFile = new FileInfo(Path.Combine(filesystemRoot, $"{debugPrefix}-{Path.GetRandomFileName()}"));
 		}
 
 		internal FilesystemBufferedPayload(FileInfo bufferFile)
@@ -123,9 +123,9 @@ namespace Jupiter.Common.Implementation
 			_length = _tempFile.Length;
 		}
 
-		public static async Task<FilesystemBufferedPayload> CreateAsync(Tracer tracer, Stream s, string filesystemRoot, CancellationToken cancellationToken)
+		public static async Task<FilesystemBufferedPayload> CreateAsync(Tracer tracer, Stream s, string filesystemRoot, string debugPrefix, CancellationToken cancellationToken)
 		{
-			FilesystemBufferedPayload payload = new FilesystemBufferedPayload(filesystemRoot);
+			FilesystemBufferedPayload payload = new FilesystemBufferedPayload(filesystemRoot, debugPrefix);
 
 			{
 				using TelemetrySpan? scope = tracer.StartActiveSpan("payload.buffer")
@@ -183,7 +183,7 @@ namespace Jupiter.Common.Implementation
 			Directory.CreateDirectory(options.CurrentValue.FilesystemTempPayloadRoot);
 		}
 
-		public Task<IBufferedPayload> CreateFromRequestAsync(HttpRequest request, CancellationToken cancellationToken)
+		public Task<IBufferedPayload> CreateFromRequestAsync(HttpRequest request, string debugPrefix, CancellationToken cancellationToken)
 		{
 			long? contentLength = request.ContentLength;
 
@@ -192,10 +192,10 @@ namespace Jupiter.Common.Implementation
 				throw new Exception("Expected content-length on all requests");
 			}
 
-			return CreateFromStreamAsync(request.Body, contentLength.Value, cancellationToken);
+			return CreateFromStreamAsync(request.Body, contentLength.Value, debugPrefix, cancellationToken);
 		}
 
-		public async Task<IBufferedPayload> CreateFromStreamAsync(Stream s, long contentLength, CancellationToken cancellationToken)
+		public async Task<IBufferedPayload> CreateFromStreamAsync(Stream s, long contentLength, string debugPrefix, CancellationToken cancellationToken)
 		{
 			// blob is small enough to fit into memory we just read it as is
 			if (contentLength < _options.CurrentValue.MemoryBufferSize)
@@ -203,12 +203,12 @@ namespace Jupiter.Common.Implementation
 				return await MemoryBufferedPayload.CreateAsync(_tracer, s, cancellationToken);
 			}
 
-			return await FilesystemBufferedPayload.CreateAsync(_tracer, s, _options.CurrentValue.FilesystemTempPayloadRoot, cancellationToken);
+			return await FilesystemBufferedPayload.CreateAsync(_tracer, s, _options.CurrentValue.FilesystemTempPayloadRoot, debugPrefix, cancellationToken);
 		}
 
-		public async Task<IBufferedPayload> CreateFilesystemBufferedPayloadAsync(Stream s, CancellationToken cancellationToken)
+		public async Task<IBufferedPayload> CreateFilesystemBufferedPayloadAsync(Stream s, string debugPrefix, CancellationToken cancellationToken)
 		{
-			return await FilesystemBufferedPayload.CreateAsync(_tracer, s, _options.CurrentValue.FilesystemTempPayloadRoot, cancellationToken);
+			return await FilesystemBufferedPayload.CreateAsync(_tracer, s, _options.CurrentValue.FilesystemTempPayloadRoot, debugPrefix, cancellationToken);
 		}
 
 		public FilesystemBufferedPayloadWriter CreateFilesystemBufferedPayloadWriter()
diff --git a/Engine/Source/Programs/UnrealCloudDDC/Jupiter/Controllers/BlobController.cs b/Engine/Source/Programs/UnrealCloudDDC/Jupiter/Controllers/BlobController.cs
--- a/Engine/Source/Programs/UnrealCloudDDC/Jupiter/Controllers/BlobController.cs
+++ b/Engine/Source/Programs/UnrealCloudDDC/Jupiter/Controllers/BlobController.cs
@@ -201,7 +201,7 @@ namespace Jupiter.Controllers
 				{
 					return Ok(new BlobUploadUriResponse(id, uri));
 				}
-				using IBufferedPayload payload = await _bufferedPayloadFactory.CreateFromRequestAsync(Request, HttpContext.RequestAborted);
+				using IBufferedPayload payload = await _bufferedPayloadFactory.CreateFromRequestAsync(Request, "blob-put", HttpContext.RequestAborted);
 
 				BlobId identifier = await _storage.PutObjectAsync(ns, payload, id, bucketHint: null, bypassCache: bypassCache, cancellationToken: HttpContext.RequestAborted);
 				return Ok(new BlobUploadResponse(identifier));
@@ -231,7 +231,7 @@ namespace Jupiter.Controllers
 			_diagnosticContext.Set("Content-Length", Request.ContentLength ?? -1);
 			try
 			{
-				using IBufferedPayload payload = await _bufferedPayloadFactory.CreateFromRequestAsync(Request, HttpContext.RequestAborted);
+				using IBufferedPayload payload = await _bufferedPayloadFactory.CreateFromRequestAsync(Request, "blob-post", HttpContext.RequestAborted);
 				bool? bypassCache = _namespacePolicyResolver.GetPoliciesForNs(ns).BypassCacheOnWrite;
 
 				await using Stream stream = payload.GetStream();
diff --git a/Engine/Source/Programs/UnrealCloudDDC/Jupiter/Controllers/BuildsController.cs b/Engine/Source/Programs/UnrealCloudDDC/Jupiter/Controllers/BuildsController.cs
--- a/Engine/Source/Programs/UnrealCloudDDC/Jupiter/Controllers/BuildsController.cs
+++ b/Engine/Source/Programs/UnrealCloudDDC/Jupiter/Controllers/BuildsController.cs
@@ -107,7 +107,7 @@ namespace Jupiter.Controllers
 			CbObject buildObject;
 			try
 			{
-				using IBufferedPayload payload = await _bufferedPayloadFactory.CreateFromRequestAsync(Request, HttpContext.RequestAborted);
+				using IBufferedPayload payload = await _bufferedPayloadFactory.CreateFromRequestAsync(Request, "put-build", HttpContext.RequestAborted);
 
 				switch (Request.ContentType)
 				{
@@ -534,7 +534,7 @@ namespace Jupiter.Controllers
 			BlobId payloadHash;
 			try
 			{
-				using IBufferedPayload payload = await _bufferedPayloadFactory.CreateFromRequestAsync(Request, HttpContext.RequestAborted);
+				using IBufferedPayload payload = await _bufferedPayloadFactory.CreateFromRequestAsync(Request, "put-build-part", HttpContext.RequestAborted);
 
 				switch (Request.ContentType)
 				{
@@ -863,10 +863,11 @@ namespace Jupiter.Controllers
 			try
 			{
 				bool? bypassCache = _namespacePolicyResolver.GetPoliciesForNs(ns).BypassCacheOnWrite;
+				
+				using IBufferedPayload payload = await _bufferedPayloadFactory.CreateFromRequestAsync(Request, "put-build-blob", HttpContext.RequestAborted);
 				if (Request.ContentType == CustomMediaTypeNames.UnrealCompressedBuffer)
 				{
 					ContentId cid = ContentId.FromBlobIdentifier(id);
-					using IBufferedPayload payload = await _bufferedPayloadFactory.CreateFromRequestAsync(Request, HttpContext.RequestAborted);
 
 					ContentId identifier = await _blobStore.PutCompressedObjectAsync(ns, payload, cid, HttpContext.RequestServices, bucketHint: bucket, bypassCache: bypassCache, cancellationToken: HttpContext.RequestAborted);
 
@@ -879,7 +880,6 @@ namespace Jupiter.Controllers
 					{
 						return Ok(new BlobUploadUriResponse(id, uri));
 					}
-					using IBufferedPayload payload = await _bufferedPayloadFactory.CreateFromRequestAsync(Request, HttpContext.RequestAborted);
 
 					BlobId identifier = await _blobStore.PutObjectAsync(ns, payload, id, bucketHint: bucket, bypassCache: bypassCache, cancellationToken: HttpContext.RequestAborted);
 					return Ok(new BlobUploadResponse(identifier));
@@ -978,7 +978,7 @@ namespace Jupiter.Controllers
 				}
 			}
 
-			using IBufferedPayload payload = await _bufferedPayloadFactory.CreateFromRequestAsync(Request, HttpContext.RequestAborted);
+			using IBufferedPayload payload = await _bufferedPayloadFactory.CreateFromRequestAsync(Request, "put-blob-multipart", HttpContext.RequestAborted);
 			await using Stream s = payload.GetStream();
 			byte[] payloadData = await s.ReadAllBytesAsync();
 			await _blobStore.PutMultipartUploadAsync(ns, blobName, uploadId, partId, payloadData);
@@ -1065,7 +1065,7 @@ namespace Jupiter.Controllers
 			
 			try
 			{
-				using IBufferedPayload payload = await _bufferedPayloadFactory.CreateFromRequestAsync(Request, HttpContext.RequestAborted);
+				using IBufferedPayload payload = await _bufferedPayloadFactory.CreateFromRequestAsync(Request, "put-block-metadata", HttpContext.RequestAborted);
 
 				await using MemoryStream ms = new MemoryStream();
 				await using Stream payloadStream = payload.GetStream();
diff --git a/Engine/Source/Programs/UnrealCloudDDC/Jupiter/Controllers/CompressedBlobController.cs b/Engine/Source/Programs/UnrealCloudDDC/Jupiter/Controllers/CompressedBlobController.cs
--- a/Engine/Source/Programs/UnrealCloudDDC/Jupiter/Controllers/CompressedBlobController.cs
+++ b/Engine/Source/Programs/UnrealCloudDDC/Jupiter/Controllers/CompressedBlobController.cs
@@ -238,7 +238,7 @@ namespace Jupiter.Controllers
 			try
 			{
 				bool? bypassCache = _namespacePolicyResolver.GetPoliciesForNs(ns).BypassCacheOnWrite;
-				using IBufferedPayload payload = await _bufferedPayloadFactory.CreateFromRequestAsync(Request, HttpContext.RequestAborted);
+				using IBufferedPayload payload = await _bufferedPayloadFactory.CreateFromRequestAsync(Request, "put-compressed", HttpContext.RequestAborted);
 
 				ContentId identifier = await _storage.PutCompressedObjectAsync(ns, payload, id, HttpContext.RequestServices, HttpContext.RequestAborted, bypassCache: bypassCache);
 
@@ -276,7 +276,7 @@ namespace Jupiter.Controllers
 
 			try
 			{
-				using IBufferedPayload payload = await _bufferedPayloadFactory.CreateFromRequestAsync(Request, cancellationToken);
+				using IBufferedPayload payload = await _bufferedPayloadFactory.CreateFromRequestAsync(Request, "post-compressed", cancellationToken);
 
 				ContentId identifier = await _storage.PutCompressedObjectAsync(ns, payload, null, HttpContext.RequestServices, cancellationToken);
 
diff --git a/Engine/Source/Programs/UnrealCloudDDC/Jupiter/Controllers/ObjectController.cs b/Engine/Source/Programs/UnrealCloudDDC/Jupiter/Controllers/ObjectController.cs
--- a/Engine/Source/Programs/UnrealCloudDDC/Jupiter/Controllers/ObjectController.cs
+++ b/Engine/Source/Programs/UnrealCloudDDC/Jupiter/Controllers/ObjectController.cs
@@ -160,7 +160,7 @@ namespace Jupiter.Controllers
 			_diagnosticContext.Set("Content-Length", Request.ContentLength ?? -1);
 			try
 			{
-				using IBufferedPayload payload = await _bufferedPayloadFactory.CreateFromRequestAsync(Request, HttpContext.RequestAborted);
+				using IBufferedPayload payload = await _bufferedPayloadFactory.CreateFromRequestAsync(Request, "put-object", HttpContext.RequestAborted);
 
 				BlobId identifier = await _storage.PutObjectAsync(ns, payload, id, bucketHint: null, cancellationToken: HttpContext.RequestAborted);
 				return Ok(new PutBlobResponse(identifier));
diff --git a/Engine/Source/Programs/UnrealCloudDDC/Jupiter/Controllers/ReferencesController.cs b/Engine/Source/Programs/UnrealCloudDDC/Jupiter/Controllers/ReferencesController.cs
--- a/Engine/Source/Programs/UnrealCloudDDC/Jupiter/Controllers/ReferencesController.cs
+++ b/Engine/Source/Programs/UnrealCloudDDC/Jupiter/Controllers/ReferencesController.cs
@@ -818,7 +818,7 @@ namespace Jupiter.Controllers
 
 			try
 			{
-				using IBufferedPayload payload = await _bufferedPayloadFactory.CreateFromRequestAsync(Request, HttpContext.RequestAborted);
+				using IBufferedPayload payload = await _bufferedPayloadFactory.CreateFromRequestAsync(Request, "put-blob-ref", HttpContext.RequestAborted);
 
 				BlobId headerHash;
 				if (Request.Headers.TryGetValue(CommonHeaders.HashHeaderName, out StringValues headers))
@@ -1119,7 +1119,7 @@ namespace Jupiter.Controllers
 				if (Request.ContentType == CustomMediaTypeNames.UnrealCompressedBuffer)
 				{
 					ContentId cid = ContentId.FromBlobIdentifier(id);
-					using IBufferedPayload payload = await _bufferedPayloadFactory.CreateFromRequestAsync(Request, HttpContext.RequestAborted);
+					using IBufferedPayload payload = await _bufferedPayloadFactory.CreateFromRequestAsync(Request, "put-blob-ref-cb", HttpContext.RequestAborted);
 
 					ContentId identifier = await _blobStore.PutCompressedObjectAsync(ns, payload, cid, HttpContext.RequestServices, bucketHint: bucket, bypassCache: bypassCache, cancellationToken: HttpContext.RequestAborted);
 
@@ -1132,7 +1132,7 @@ namespace Jupiter.Controllers
 					{
 						return Ok(new BlobUploadUriResponse(id, uri));
 					}
-					using IBufferedPayload payload = await _bufferedPayloadFactory.CreateFromRequestAsync(Request, HttpContext.RequestAborted);
+					using IBufferedPayload payload = await _bufferedPayloadFactory.CreateFromRequestAsync(Request, "put-blob-ref-raw", HttpContext.RequestAborted);
 
 					BlobId identifier = await _blobStore.PutObjectAsync(ns, payload, id, bucketHint: bucket, bypassCache: bypassCache, cancellationToken: HttpContext.RequestAborted);
 					return Ok(new  BlobUploadResponse(identifier));
diff --git a/Engine/Source/Programs/UnrealCloudDDC/Jupiter/Controllers/SymbolsController.cs b/Engine/Source/Programs/UnrealCloudDDC/Jupiter/Controllers/SymbolsController.cs
--- a/Engine/Source/Programs/UnrealCloudDDC/Jupiter/Controllers/SymbolsController.cs
+++ b/Engine/Source/Programs/UnrealCloudDDC/Jupiter/Controllers/SymbolsController.cs
@@ -191,7 +191,7 @@ namespace Jupiter.Controllers
 			_diagnosticContext.Set("Content-Length", Request.ContentLength ?? -1);
 			CompressedBufferUtils utils = new CompressedBufferUtils(_tracer, _bufferedPayloadFactory);
 
-			IBufferedPayload payloadToUse = await _bufferedPayloadFactory.CreateFromRequestAsync(Request, HttpContext.RequestAborted);
+			IBufferedPayload payloadToUse = await _bufferedPayloadFactory.CreateFromRequestAsync(Request, "symbol-put", HttpContext.RequestAborted);
 			if (Request.ContentType == MediaTypeNames.Application.Octet)
 			{
 				await using Stream s = payloadToUse.GetStream();
@@ -201,7 +201,7 @@ namespace Jupiter.Controllers
 				utils.CompressContent(compressedStream, OoodleCompressorMethod.Kraken, OoodleCompressionLevel.VeryFast, await s.ToByteArrayAsync(HttpContext.RequestAborted));
 
 				compressedStream.Seek(0, SeekOrigin.Begin);
-				payloadToUse = await _bufferedPayloadFactory.CreateFromStreamAsync(compressedStream, compressedStream.Length, HttpContext.RequestAborted);
+				payloadToUse = await _bufferedPayloadFactory.CreateFromStreamAsync(compressedStream, compressedStream.Length, "symbol-put-recompress", HttpContext.RequestAborted);
 			}
 
 			BucketId bucketName = new BucketId(moduleName);
diff --git a/Engine/Source/Programs/UnrealCloudDDC/Jupiter/Implementation/Blob/BlobService.cs b/Engine/Source/Programs/UnrealCloudDDC/Jupiter/Implementation/Blob/BlobService.cs
--- a/Engine/Source/Programs/UnrealCloudDDC/Jupiter/Implementation/Blob/BlobService.cs
+++ b/Engine/Source/Programs/UnrealCloudDDC/Jupiter/Implementation/Blob/BlobService.cs
@@ -244,7 +244,7 @@ public class BlobService : IBlobService
 	private async Task<BlobId> PutObjectToStoresAsync(NamespaceId ns, byte[] payload, BlobId identifier, BucketId? bucketHint = null, bool? bypassCache = null, CancellationToken cancellationToken = default)
 	{
 		using MemoryStream ms = new MemoryStream(payload);
-		using IBufferedPayload bufferedPayload = await _bufferedPayloadFactory.CreateFromStreamAsync(ms, ms.Length, cancellationToken);
+		using IBufferedPayload bufferedPayload = await _bufferedPayloadFactory.CreateFromStreamAsync(ms, ms.Length, "put-object-store", cancellationToken);
 		return await PutObjectToStoresAsync(ns, bufferedPayload, identifier, bucketHint, bypassCache, cancellationToken);
 	}
 
@@ -430,7 +430,7 @@ public class BlobService : IBlobService
 			using ServerTimingMetricScoped? serverTimingScope = serverTiming?.CreateServerTimingMetricScope($"blob.populate", "Populating caches with blob contents");
 
 			// not using using as the blob contents will take ownership of this buffered payload and dispose it when the contents is disposed
-			IBufferedPayload bufferedPayload = await _bufferedPayloadFactory.CreateFromStreamAsync(blobContents.Stream, blobContents.Length, cancellationToken);
+			IBufferedPayload bufferedPayload = await _bufferedPayloadFactory.CreateFromStreamAsync(blobContents.Stream, blobContents.Length, "populate-store", cancellationToken);
 
 			// Don't populate the last store, as that is where we got the hit
 			for (int i = 0; i < numStoreMisses; i++)
@@ -633,7 +633,7 @@ public class BlobService : IBlobService
 			}
 
 			await using Stream s = await blobResponse.Content.ReadAsStreamAsync(cancellationToken);
-			using IBufferedPayload payload = await _bufferedPayloadFactory.CreateFilesystemBufferedPayloadAsync(s, cancellationToken);
+			using IBufferedPayload payload = await _bufferedPayloadFactory.CreateFilesystemBufferedPayloadAsync(s, "replicate-blob", cancellationToken);
 			await PutObjectAsync(ns, payload, blob, bucketHint: bucketHint, cancellationToken: cancellationToken);
 			replicated = true;
 			break;
@@ -784,7 +784,7 @@ public class BlobService : IBlobService
 			throw new Exception($"Failed to find multipart blob {blobName}");
 		}
 
-		using IBufferedPayload bufferedPayload = await _bufferedPayloadFactory.CreateFromStreamAsync(blobContents.Stream, blobContents.Length, cancellationToken);
+		using IBufferedPayload bufferedPayload = await _bufferedPayloadFactory.CreateFromStreamAsync(blobContents.Stream, blobContents.Length, "verify-multipart", cancellationToken);
 		ContentId cid;
 		BlobId blobId;
 		if (isCompressed)
diff --git a/Engine/Source/Programs/UnrealCloudDDC/Jupiter/Implementation/Replication/BlobsReplicator.cs b/Engine/Source/Programs/UnrealCloudDDC/Jupiter/Implementation/Replication/BlobsReplicator.cs
--- a/Engine/Source/Programs/UnrealCloudDDC/Jupiter/Implementation/Replication/BlobsReplicator.cs
+++ b/Engine/Source/Programs/UnrealCloudDDC/Jupiter/Implementation/Replication/BlobsReplicator.cs
@@ -504,7 +504,7 @@ namespace Jupiter.Implementation
 				throw new Exception("Expected content-length on blob response");
 			}
 
-			using IBufferedPayload payload = await _bufferedPayloadFactory.CreateFromStreamAsync(s, contentLength.Value, cancellationToken);
+			using IBufferedPayload payload = await _bufferedPayloadFactory.CreateFromStreamAsync(s, contentLength.Value, "blob-replicator", cancellationToken);
 
 			await _blobService.PutObjectAsync(ns, payload, blob, bucketHint: bucketHint, cancellationToken: cancellationToken);
 
diff --git a/Engine/Source/Programs/UnrealCloudDDC/Jupiter/Implementation/Replication/RefsReplicator.cs b/Engine/Source/Programs/UnrealCloudDDC/Jupiter/Implementation/Replication/RefsReplicator.cs
--- a/Engine/Source/Programs/UnrealCloudDDC/Jupiter/Implementation/Replication/RefsReplicator.cs
+++ b/Engine/Source/Programs/UnrealCloudDDC/Jupiter/Implementation/Replication/RefsReplicator.cs
@@ -596,7 +596,7 @@ namespace Jupiter.Implementation
 						throw new Exception("Expected content-length on blob response");
 					}
 
-					using IBufferedPayload payload = await _bufferedPayloadFactory.CreateFromStreamAsync(s, contentLength.Value, cancellationToken);
+					using IBufferedPayload payload = await _bufferedPayloadFactory.CreateFromStreamAsync(s, contentLength.Value, "ref-replicator", cancellationToken);
 
 					await _blobService.PutObjectAsync(ns, payload, blobToReplicate, bucketHint: bucket, cancellationToken: cancellationToken);
 				}, cancellationToken);
diff --git a/Engine/Source/Programs/UnrealCloudDDC/Jupiter/Implementation/TransactionLog/ReplicationLogSnapshot.cs b/Engine/Source/Programs/UnrealCloudDDC/Jupiter/Implementation/TransactionLog/ReplicationLogSnapshot.cs
--- a/Engine/Source/Programs/UnrealCloudDDC/Jupiter/Implementation/TransactionLog/ReplicationLogSnapshot.cs
+++ b/Engine/Source/Programs/UnrealCloudDDC/Jupiter/Implementation/TransactionLog/ReplicationLogSnapshot.cs
@@ -82,7 +82,7 @@ namespace Jupiter.Implementation.TransactionLog
 
 		public static ReplicationLogSnapshot FromStream(BufferedPayloadFactory payloadFactory, Stream stream)
 		{
-			IBufferedPayload payload = payloadFactory.CreateFilesystemBufferedPayloadAsync(stream, CancellationToken.None).Result;
+			IBufferedPayload payload = payloadFactory.CreateFilesystemBufferedPayloadAsync(stream, "snapshot-serialize", CancellationToken.None).Result;
 
 			using Stream payloadStream = payload.GetStream();
 			(NamespaceId ns, string lastBucket, Guid lastEvent, ulong countOfObjects) = ReadHeader(payloadStream);
diff --git a/Engine/Source/Programs/UnrealCloudDDC/Jupiter/Implementation/TransactionLog/ReplicationLogSnapshotBuilder.cs b/Engine/Source/Programs/UnrealCloudDDC/Jupiter/Implementation/TransactionLog/ReplicationLogSnapshotBuilder.cs
--- a/Engine/Source/Programs/UnrealCloudDDC/Jupiter/Implementation/TransactionLog/ReplicationLogSnapshotBuilder.cs
+++ b/Engine/Source/Programs/UnrealCloudDDC/Jupiter/Implementation/TransactionLog/ReplicationLogSnapshotBuilder.cs
@@ -58,7 +58,7 @@ namespace Jupiter.Implementation.TransactionLog
 					}
 
 					using IBufferedPayload snapshotPayload =
-						await _bufferedPayloadFactory.CreateFilesystemBufferedPayloadAsync(blobContents.Stream, cancellationToken);
+						await _bufferedPayloadFactory.CreateFilesystemBufferedPayloadAsync(blobContents.Stream, "snapshot-build", cancellationToken);
 					await using Stream s = snapshotPayload.GetStream();
 					snapshot = _replicationLogFactory.DeserializeSnapshotFromStream(s);
 					lastBucket = snapshot.LastBucket;
@@ -96,7 +96,7 @@ namespace Jupiter.Implementation.TransactionLog
 				}
 
 				Stream tempFileStream = tempFile.OpenRead();
-				using IBufferedPayload payload = await _bufferedPayloadFactory.CreateFromStreamAsync(tempFileStream, tempFile.Length, cancellationToken);
+				using IBufferedPayload payload = await _bufferedPayloadFactory.CreateFromStreamAsync(tempFileStream, tempFile.Length, "snapshot-serialize", cancellationToken);
 				tempFileStream.Close();
 
 				{
