Cloud DDC - Added more temp storage debug prefixes to help catch unreleased temp payloads.

diff --git a/Engine/Source/Programs/UnrealCloudDDC/Jupiter.Common/Implementation/StreamingBodyPayload.cs b/Engine/Source/Programs/UnrealCloudDDC/Jupiter.Common/Implementation/StreamingBodyPayload.cs
--- a/Engine/Source/Programs/UnrealCloudDDC/Jupiter.Common/Implementation/StreamingBodyPayload.cs
+++ b/Engine/Source/Programs/UnrealCloudDDC/Jupiter.Common/Implementation/StreamingBodyPayload.cs
@@ -59,9 +59,9 @@ namespace Jupiter.Common.Implementation
 	{
 		private FileInfo? _tempFile;
 
-		private FilesystemBufferedPayloadWriter(string filesystemRoot)
+		private FilesystemBufferedPayloadWriter(string filesystemRoot, string debugPrefix)
 		{
-			_tempFile = new FileInfo(Path.Combine(filesystemRoot, Path.GetRandomFileName()));
+			_tempFile = new FileInfo(Path.Combine(filesystemRoot,$"{debugPrefix}-{Path.GetRandomFileName()}"));
 		}
 
 		public void Dispose()
@@ -95,9 +95,9 @@ namespace Jupiter.Common.Implementation
 			return _tempFile.OpenWrite();
 		}
 
-		public static FilesystemBufferedPayloadWriter Create(string filesystemTempPayloadRoot)
+		public static FilesystemBufferedPayloadWriter Create(string filesystemTempPayloadRoot, string debugPrefix)
 		{
-			return new FilesystemBufferedPayloadWriter(filesystemTempPayloadRoot);
+			return new FilesystemBufferedPayloadWriter(filesystemTempPayloadRoot, debugPrefix);
 		}
 	}
 
@@ -211,9 +211,9 @@ namespace Jupiter.Common.Implementation
 			return await FilesystemBufferedPayload.CreateAsync(_tracer, s, _options.CurrentValue.FilesystemTempPayloadRoot, debugPrefix, cancellationToken);
 		}
 
-		public FilesystemBufferedPayloadWriter CreateFilesystemBufferedPayloadWriter()
+		public FilesystemBufferedPayloadWriter CreateFilesystemBufferedPayloadWriter(string debugPrefix)
 		{
-			return FilesystemBufferedPayloadWriter.Create(_options.CurrentValue.FilesystemTempPayloadRoot);
+			return FilesystemBufferedPayloadWriter.Create(_options.CurrentValue.FilesystemTempPayloadRoot, debugPrefix);
 		}
 	}
 }
diff --git a/Engine/Source/Programs/UnrealCloudDDC/Jupiter/Implementation/Blob/AmazonS3Store.cs b/Engine/Source/Programs/UnrealCloudDDC/Jupiter/Implementation/Blob/AmazonS3Store.cs
--- a/Engine/Source/Programs/UnrealCloudDDC/Jupiter/Implementation/Blob/AmazonS3Store.cs
+++ b/Engine/Source/Programs/UnrealCloudDDC/Jupiter/Implementation/Blob/AmazonS3Store.cs
@@ -424,7 +424,7 @@ namespace Jupiter.Implementation
 				else if (stream.Length > 16 * (long)Math.Pow(2, 20))
 				{
 					// will be chunked by TransferUtility
-					using FilesystemBufferedPayloadWriter writer = _payloadFactory.CreateFilesystemBufferedPayloadWriter();
+					using FilesystemBufferedPayloadWriter writer = _payloadFactory.CreateFilesystemBufferedPayloadWriter("s3-upload");
 					{
 						await using Stream writableStream = writer.GetWritableStream();
 						await stream.CopyToAsync(writableStream, cancellationToken);
@@ -508,7 +508,7 @@ namespace Jupiter.Implementation
 							lengths[part-1] = partResponse.ContentLength;
 						});
 
-						using FilesystemBufferedPayloadWriter writer = _payloadFactory.CreateFilesystemBufferedPayloadWriter();
+						using FilesystemBufferedPayloadWriter writer = _payloadFactory.CreateFilesystemBufferedPayloadWriter("s3-download");
 
 						{
 							await using Stream s = writer.GetWritableStream();
diff --git a/Engine/Source/Programs/UnrealCloudDDC/Jupiter/Implementation/ContentId/CompressedBufferUtils.cs b/Engine/Source/Programs/UnrealCloudDDC/Jupiter/Implementation/ContentId/CompressedBufferUtils.cs
--- a/Engine/Source/Programs/UnrealCloudDDC/Jupiter/Implementation/ContentId/CompressedBufferUtils.cs
+++ b/Engine/Source/Programs/UnrealCloudDDC/Jupiter/Implementation/ContentId/CompressedBufferUtils.cs
@@ -194,7 +194,7 @@ namespace Jupiter.Implementation
 				throw new Exception($"Expected stream to be {header.TotalCompressedSize} but it was {streamSize}");
 			}
 
-			using FilesystemBufferedPayloadWriter bufferedPayloadWriter = _payloadFactory.CreateFilesystemBufferedPayloadWriter();
+			using FilesystemBufferedPayloadWriter bufferedPayloadWriter = _payloadFactory.CreateFilesystemBufferedPayloadWriter("cb-decompress");
 
 			{
 				await using Stream targetStream = bufferedPayloadWriter.GetWritableStream();
