Cloud DDC - Fixed issue were temporary payloads could be kept around when promoted into the filesystem cache.
Also made s3 multipart buffer to memory if possible avoiding temp files.

diff --git a/Engine/Source/Programs/UnrealCloudDDC/Jupiter/Implementation/Blob/AmazonS3Store.cs b/Engine/Source/Programs/UnrealCloudDDC/Jupiter/Implementation/Blob/AmazonS3Store.cs
--- a/Engine/Source/Programs/UnrealCloudDDC/Jupiter/Implementation/Blob/AmazonS3Store.cs
+++ b/Engine/Source/Programs/UnrealCloudDDC/Jupiter/Implementation/Blob/AmazonS3Store.cs
@@ -508,18 +508,35 @@ namespace Jupiter.Implementation
 							lengths[part-1] = partResponse.ContentLength;
 						});
 
-						using FilesystemBufferedPayloadWriter writer = _payloadFactory.CreateFilesystemBufferedPayloadWriter("s3-download");
-
+						long totalSize = lengths.Sum();
+						if (totalSize < int.MaxValue)
 						{
-							await using Stream s = writer.GetWritableStream();
+							Stream s = new MemoryStream(new byte[(int)totalSize]);
 							for (int i = 0; i < parts.Length; i++)
 							{
 								byte[] part = parts[i];
 								await s.WriteAsync(part, 0, (int)lengths[i], cancellationToken);
 							}
+
+							s.Position = 0;
+							return new BlobContents(s, totalSize);
 						}
+						else
+						{
+							// large payload - we have to buffer to file
+							using FilesystemBufferedPayloadWriter writer = _payloadFactory.CreateFilesystemBufferedPayloadWriter("s3-download");
+
+							{
+								await using Stream s = writer.GetWritableStream();
+								for (int i = 0; i < parts.Length; i++)
+								{
+									byte[] part = parts[i];
+									await s.WriteAsync(part, 0, (int)lengths[i], cancellationToken);
+								}
+							}
 						
-						return new BlobContents(writer.Done());
+							return new BlobContents(writer.Done());
+						}
 					}
 					finally
 					{
diff --git a/Engine/Source/Programs/UnrealCloudDDC/Jupiter/Implementation/Blob/BlobService.cs b/Engine/Source/Programs/UnrealCloudDDC/Jupiter/Implementation/Blob/BlobService.cs
--- a/Engine/Source/Programs/UnrealCloudDDC/Jupiter/Implementation/Blob/BlobService.cs
+++ b/Engine/Source/Programs/UnrealCloudDDC/Jupiter/Implementation/Blob/BlobService.cs
@@ -429,7 +429,7 @@ public class BlobService : IBlobService
 			using TelemetrySpan _ = _tracer.StartActiveSpan("HierarchicalStore.Populate").SetAttribute("operation.name", "HierarchicalStore.Populate");
 			using ServerTimingMetricScoped? serverTimingScope = serverTiming?.CreateServerTimingMetricScope($"blob.populate", "Populating caches with blob contents");
 
-			// not using using as the blob contents will take ownership of this buffered payload and dispose it when the contents is disposed
+			// no using statement as the blob contents will take ownership of this buffered payload and dispose it when the contents it is disposed
 			IBufferedPayload bufferedPayload = await _bufferedPayloadFactory.CreateFromStreamAsync(blobContents.Stream, blobContents.Length, "populate-store", cancellationToken);
 
 			// Don't populate the last store, as that is where we got the hit
@@ -444,9 +444,9 @@ public class BlobService : IBlobService
 				await blobStore.PutObjectAsync(ns, s, blob);
 			}
 
-#pragma warning disable CA2000 // Dispose objects before losing scope , ownership is transfered to caller
-			blobContents = new BlobContents(bufferedPayload);
-#pragma warning restore CA2000 // Dispose objects before losing scope
+			// dispose the old blob contents before returning a new one
+			await blobContents.DisposeAsync();
+			return new BlobContents(bufferedPayload);
 		}
 
 		return blobContents;
