Cloud DDC - Check if block metadata is reported as not existing in the block index when list blocks if we were able to fetch the contents (indicates the block existince in the filesystem cache but has been deleted from the root store). If that happens do not report the block as existing which will force the block contents to be reuploaded. Should fix inconsistency issues between getBlockMetadata and listBlocks.

diff --git a/Engine/Source/Programs/UnrealCloudDDC/Jupiter/Controllers/BuildsController.cs b/Engine/Source/Programs/UnrealCloudDDC/Jupiter/Controllers/BuildsController.cs
--- a/Engine/Source/Programs/UnrealCloudDDC/Jupiter/Controllers/BuildsController.cs
+++ b/Engine/Source/Programs/UnrealCloudDDC/Jupiter/Controllers/BuildsController.cs
@@ -1175,7 +1175,7 @@ namespace Jupiter.Controllers
 				}
 				catch (BlobNotFoundException)
 				{
-					// if the metadata doesnt exist we can not display it
+					// if the metadata doesn't exist we can not display it
 				}
 			});
 			responseObject.EndArray();
@@ -1252,10 +1252,18 @@ namespace Jupiter.Controllers
 				await Parallel.ForEachAsync(enumerable.Skip(skip).Take(count), async (pair, token) =>
 				{
 					(DateTime _, BlobId blockMetadata) = pair;
+					Task<bool> blobExistsTask = _blobIndex.BlobExistsInRegionAsync(ns, blockMetadata, cancellationToken: token);
 					try
 					{
 						await using BlobContents contents = await _blobStore.GetObjectAsync(ns, blockMetadata, cancellationToken: token);
 						byte[] b = await contents.Stream.ReadAllBytesAsync(cancellationToken: token);
+						if (!await blobExistsTask)
+						{
+							// we got blob contents but the blob is reported to not exist in the index
+							// to recover we flag this block as not existing so that its content is uploaded again
+							return;
+						}
+
 						CbObject o = new CbObject(b);
 						lock (writer)
 						{
