Cloud DDC - Added endpoint to forcefully delete a compressed buffer blob (already exists for the raw blob). Not generally useful as GC is expected to handle this but can be good to help recover if we have corrupt data (due to bugs).

diff --git a/Engine/Source/Programs/UnrealCloudDDC/Jupiter/Controllers/CompressedBlobController.cs b/Engine/Source/Programs/UnrealCloudDDC/Jupiter/Controllers/CompressedBlobController.cs
--- a/Engine/Source/Programs/UnrealCloudDDC/Jupiter/Controllers/CompressedBlobController.cs
+++ b/Engine/Source/Programs/UnrealCloudDDC/Jupiter/Controllers/CompressedBlobController.cs
@@ -1,5 +1,6 @@
 ﻿// Copyright Epic Games, Inc. All Rights Reserved.
 
+using System;
 using System.Collections.Concurrent;
 using System.Collections.Generic;
 using System.ComponentModel.DataAnnotations;
@@ -298,43 +299,39 @@ namespace Jupiter.Controllers
 			}
 		}
 
-		/*[HttpDelete("{ns}/{id}")]
-		public async Task<IActionResult> Delete(
-			[Required] string ns,
-			[Required] BlobIdentifier id)
+		[HttpDelete("{ns}/{id}")]
+		public async Task<IActionResult> DeleteAsync(
+			[Required] NamespaceId ns,
+			[Required] ContentId id)
 		{
-			ActionResult? result = await _requestHelper.HasAccessToNamespace(User, Request, ns);
+			ActionResult? result = await _requestHelper.HasAccessToNamespaceAsync(User, Request, ns, new[] { JupiterAclAction.DeleteObject });
 			if (result != null)
 			{
 				return result;
 			}
 
-			await DeleteImpl(ns, id);
-
-			return Ok();
-		}
+			BlobId[]? chunks = await _contentIdStore.ResolveAsync(ns, id, mustBeContentId: true);
 
-		
-		[HttpDelete("{ns}")]
-		public async Task<IActionResult> DeleteNamespace(
-			[Required] string ns)
-		{
-			ActionResult? result = await _requestHelper.HasAccessToNamespace(User, Request, ns);
-			if (result != null)
+			if (chunks == null)
 			{
-				return result;
+				return NotFound();
 			}
 
-			int deletedCount = await  _storage.DeleteNamespace(ns);
+			foreach (BlobId chunk in chunks)
+			{
+				try
+				{
+					await _storage.DeleteObjectAsync(ns, chunk);
+				}
+				catch (BlobNotFoundException)
+				{
+					// if the blob is already missing that is fine
+				}
+			}
 
-			return Ok( new { Deleted = deletedCount });
+			// TODO: we should delete the cid from the content id store as well, but it has no delete operations yet
+			return Ok();
 		}
-
-
-		private async Task DeleteImpl(string ns, BlobIdentifier id)
-		{
-			await _storage.DeleteObject(ns, id);
-		}*/
 	}
 
 	public class ExistCheckMultipleContentIdResponse
