Cloud DDC - Reverted back block listing to statement before change 39518517 - after discussions with Scylla its clear that using a index like this is completley fine and should work, and issues that comes up should be reported as bugs. There is an issue with this in Scylla 6.2 though so reverting our test images to 6.1 while Scylla works on a fix.

diff --git a/Engine/Source/Programs/UnrealCloudDDC/Composes/docker-compose-tests.yml b/Engine/Source/Programs/UnrealCloudDDC/Composes/docker-compose-tests.yml
--- a/Engine/Source/Programs/UnrealCloudDDC/Composes/docker-compose-tests.yml
+++ b/Engine/Source/Programs/UnrealCloudDDC/Composes/docker-compose-tests.yml
@@ -40,7 +40,7 @@ services:
     ports:
     - "11002:11002"   
   scylla:
-    image: scylladb/scylla
+    image: scylladb/scylla:6.1
     command: --smp=1 --developer-mode 1
     # Use to disable logging which can be spammy but also useful
     #logging:
diff --git a/Engine/Source/Programs/UnrealCloudDDC/Composes/docker-compose.yml b/Engine/Source/Programs/UnrealCloudDDC/Composes/docker-compose.yml
--- a/Engine/Source/Programs/UnrealCloudDDC/Composes/docker-compose.yml
+++ b/Engine/Source/Programs/UnrealCloudDDC/Composes/docker-compose.yml
@@ -18,7 +18,7 @@ services:
   scylla:
     logging:
       driver: none
-    image: scylladb/scylla
+    image: scylladb/scylla:6.1
     command: --smp=1 --developer-mode 1
     volumes:
     - scylla-volume-1:/var/lib/scylla
diff --git a/Engine/Source/Programs/UnrealCloudDDC/Jupiter/Implementation/Builds/ScyllaBlockStore.cs b/Engine/Source/Programs/UnrealCloudDDC/Jupiter/Implementation/Builds/ScyllaBlockStore.cs
--- a/Engine/Source/Programs/UnrealCloudDDC/Jupiter/Implementation/Builds/ScyllaBlockStore.cs
+++ b/Engine/Source/Programs/UnrealCloudDDC/Jupiter/Implementation/Builds/ScyllaBlockStore.cs
@@ -53,7 +53,7 @@ public class ScyllaBlockStore : IBlockStore
 		using TelemetrySpan scope = _tracer.BuildScyllaSpan("scylla.list_blocks");
 
 		foreach (ScyllaBlockContextEntry blockEntry in await _mapper.FetchAsync<ScyllaBlockContextEntry>(
-			         "SELECT * FROM block_context WHERE namespace = ? AND block_context = ?", 
+			         "SELECT * FROM block_context_by_time_index WHERE namespace = ? AND block_context = ? ORDER BY last_update DESC", 
 			         ns.ToString(), blockContext.ToString()))
 		{
 			yield return new BlockMetadata { LastUpdate = blockEntry.LastUpdate, MetadataBlobId = new BlobId(blockEntry.MetadataBlobId) };
