Cloud DDC - Fixed issue were blob not found exception was not raised when blob was missing when using Scylla backed blob indexes but was for the others. Resolves failing tests.

diff --git a/Engine/Source/Programs/UnrealCloudDDC/Jupiter/Implementation/Blob/ScyllaBlobIndex.cs b/Engine/Source/Programs/UnrealCloudDDC/Jupiter/Implementation/Blob/ScyllaBlobIndex.cs
--- a/Engine/Source/Programs/UnrealCloudDDC/Jupiter/Implementation/Blob/ScyllaBlobIndex.cs
+++ b/Engine/Source/Programs/UnrealCloudDDC/Jupiter/Implementation/Blob/ScyllaBlobIndex.cs
@@ -205,7 +205,7 @@ public class ScyllaBlobIndex : IBlobIndex
 			regions.Add(blobIndex.Region);
 		}
 
-		if (regions.Any() && _scyllaSettings.CurrentValue.MigrateFromOldBlobIndex)
+		if (regions.Count == 0 && _scyllaSettings.CurrentValue.MigrateFromOldBlobIndex)
 		{
 			List<string>? oldRegions = await GetOldBlobRegionsAsync(ns, blob);
 			if (oldRegions == null)
@@ -220,8 +220,13 @@ public class ScyllaBlobIndex : IBlobIndex
 			}
 
 			// blob has been migrated lets return the old region list
-			return oldRegions;
+			regions = oldRegions;
 		}
+		if (regions.Count == 0)
+		{
+			throw new BlobNotFoundException(ns, blob);
+		}
+
 		return regions;
 	}
 
