Cloud DDC - Fixed issues were exception was not being caught for cases were blobs could be not found.

diff --git a/Engine/Source/Programs/UnrealCloudDDC/Jupiter/Implementation/Blob/BlobIndexConsistencyCheckService.cs b/Engine/Source/Programs/UnrealCloudDDC/Jupiter/Implementation/Blob/BlobIndexConsistencyCheckService.cs
--- a/Engine/Source/Programs/UnrealCloudDDC/Jupiter/Implementation/Blob/BlobIndexConsistencyCheckService.cs
+++ b/Engine/Source/Programs/UnrealCloudDDC/Jupiter/Implementation/Blob/BlobIndexConsistencyCheckService.cs
@@ -89,7 +89,15 @@ namespace Jupiter.Implementation
 
 					bool issueFound = false;
 					bool deleted = false;
-					List<string> regions = await _blobIndex.GetBlobRegionsAsync(ns, blobIdentifier, cancellationToken);
+					List<string> regions;
+					try
+					{
+						regions = await _blobIndex.GetBlobRegionsAsync(ns, blobIdentifier, cancellationToken);
+					}
+					catch (BlobNotFoundException)
+					{
+						regions = new List<string>();
+					}
 					try
 					{
 						if (regions.Contains(currentRegion))
diff --git a/Engine/Source/Programs/UnrealCloudDDC/Jupiter/Implementation/Blob/BlobService.cs b/Engine/Source/Programs/UnrealCloudDDC/Jupiter/Implementation/Blob/BlobService.cs
--- a/Engine/Source/Programs/UnrealCloudDDC/Jupiter/Implementation/Blob/BlobService.cs
+++ b/Engine/Source/Programs/UnrealCloudDDC/Jupiter/Implementation/Blob/BlobService.cs
@@ -589,7 +589,15 @@ public class BlobService : IBlobService
 
 		using ServerTimingMetricScoped? serverTimingScope = serverTiming?.CreateServerTimingMetricScope("blob.replicate", "Replicating blob from remote instances");
 
-		List<string> regions = await _blobIndex.GetBlobRegionsAsync(ns, blob, cancellationToken);
+		List<string> regions;
+		try
+		{
+			regions = await _blobIndex.GetBlobRegionsAsync(ns, blob, cancellationToken);
+		}
+		catch (BlobNotFoundException)
+		{
+			regions = new List<string>();
+		}
 
 		if (!regions.Any())
 		{
@@ -956,7 +964,15 @@ public class BlobService : IBlobService
 
 		using ServerTimingMetricScoped? serverTimingScope = serverTiming?.CreateServerTimingMetricScope("blob.exists-remote", "Verify if blob exists in remotes");
 
-		List<string> regions = await _blobIndex.GetBlobRegionsAsync(ns, blob, cancellationToken);
+		List<string> regions;
+		try
+		{
+			regions = await _blobIndex.GetBlobRegionsAsync(ns, blob, cancellationToken);
+		}
+		catch (BlobNotFoundException)
+		{
+			regions = new List<string>();
+		}
 
 		// we do not actually verify that the blob exists remotely as that would take a lot of time
 		// instead we simply check if there are any regions were the blob exists that is not our current region
