Cloud DDC - Added option to control when we buffer s3 multipart downloads to memory vs filesystem. Useful to control tradeoff between memory usage and filesystem load. Default value unchanged (we buffer to memory as much as we can).

diff --git a/Engine/Source/Programs/UnrealCloudDDC/Helm/UnrealCloudDDC/templates/configmap.yaml b/Engine/Source/Programs/UnrealCloudDDC/Helm/UnrealCloudDDC/templates/configmap.yaml
--- a/Engine/Source/Programs/UnrealCloudDDC/Helm/UnrealCloudDDC/templates/configmap.yaml
+++ b/Engine/Source/Programs/UnrealCloudDDC/Helm/UnrealCloudDDC/templates/configmap.yaml
@@ -226,5 +226,9 @@ S3:
   {{- if hasKey $s3 "UseChunkEncoding" }}
   UseChunkEncoding: {{ $s3.UseChunkEncoding }}
   {{- end }}
+
+  {{- if hasKey $s3 "MultiPartMaxMemoryBufferSize" }}
+  MultiPartMaxMemoryBufferSize: {{ $s3.MultiPartMaxMemoryBufferSize }}
+  {{- end }}
   
 {{ end }}
diff --git a/Engine/Source/Programs/UnrealCloudDDC/Jupiter/Implementation/Blob/AmazonS3Store.cs b/Engine/Source/Programs/UnrealCloudDDC/Jupiter/Implementation/Blob/AmazonS3Store.cs
--- a/Engine/Source/Programs/UnrealCloudDDC/Jupiter/Implementation/Blob/AmazonS3Store.cs
+++ b/Engine/Source/Programs/UnrealCloudDDC/Jupiter/Implementation/Blob/AmazonS3Store.cs
@@ -509,7 +509,7 @@ namespace Jupiter.Implementation
 						});
 
 						long totalSize = lengths.Sum();
-						if (totalSize < int.MaxValue)
+						if (totalSize < _settings.CurrentValue.MultiPartMaxMemoryBufferSize)
 						{
 							Stream s = new MemoryStream(new byte[(int)totalSize]);
 							for (int i = 0; i < parts.Length; i++)
diff --git a/Engine/Source/Programs/UnrealCloudDDC/Jupiter/Settings.cs b/Engine/Source/Programs/UnrealCloudDDC/Jupiter/Settings.cs
--- a/Engine/Source/Programs/UnrealCloudDDC/Jupiter/Settings.cs
+++ b/Engine/Source/Programs/UnrealCloudDDC/Jupiter/Settings.cs
@@ -227,6 +227,12 @@ namespace Jupiter
 		/// Max number of keys returned in a single request when listing S3
 		/// </summary>
 		public int PerPrefixMaxKeys { get; set; } = 10_000;
+
+		/// <summary>
+		/// Max size of bytes that is buffered into memory before switching to filesystem buffers
+		/// Larger value means more RAM used but less load on the filesystem.
+		/// </summary>
+		public int MultiPartMaxMemoryBufferSize { get; set; } = int.MaxValue;
 	}
 
 	public class GCSettings
