Cloud DDC - More debug output for search logs to be try and track down incorrect search results.

diff --git a/Engine/Source/Programs/UnrealCloudDDC/Jupiter/Implementation/Builds/SearchOp.cs b/Engine/Source/Programs/UnrealCloudDDC/Jupiter/Implementation/Builds/SearchOp.cs
--- a/Engine/Source/Programs/UnrealCloudDDC/Jupiter/Implementation/Builds/SearchOp.cs
+++ b/Engine/Source/Programs/UnrealCloudDDC/Jupiter/Implementation/Builds/SearchOp.cs
@@ -150,7 +150,7 @@ namespace Jupiter.Implementation.Builds
 				}
 
 				int compare = Comparer.Default.Compare(expectedType, actualType);
-				logger.LogDebug("Running compare between values {Expected} {Actual} result was {Result}", expectedType, actualType, compare);
+				logger.LogDebug("Running compare between values {Value0} string and cb field {Value1} result was {Result}", rawText, actualType, compare);
 				switch (op)
 				{
 					case OpType.GreaterThen:
@@ -169,7 +169,7 @@ namespace Jupiter.Implementation.Builds
 			{
 				// looks for expected value in the array
 				JsonElement[] array = expectedField.EnumerateArray().ToArray();
-				bool isInArray = array.Any(element => FieldEquals(field, element));
+				bool isInArray = array.Any(element => FieldEquals(field, element, logger));
 				
 				if (op == OpType.In)
 				{
@@ -183,9 +183,9 @@ namespace Jupiter.Implementation.Builds
 			}
 			else if (isEquatableOp)
 			{
-				bool isEqual = FieldEquals(field, expectedField);
+				bool isEqual = FieldEquals(field, expectedField, logger);
 				
-				logger.LogDebug("Ran equality check between values {Expected} {Actual} result was {Result}", field, expectedField, isEqual);
+				logger.LogDebug("Ran equality check between values from json: {Value0} and cb of type: {Value1} result was {Result}", expectedField, field.GetType().ToString(), isEqual);
 				if (op == OpType.Equals)
 				{
 					return isEqual;
@@ -199,39 +199,54 @@ namespace Jupiter.Implementation.Builds
 			throw new UnsupportedOperationException($"Field {field.Name} not convertible to a type that supports operation: {op} against type {expectedField.ValueKind}");
 		}
 
-		private static bool FieldEquals(CbField field, JsonElement jsonElement)
+		private static bool FieldEquals(CbField field, JsonElement jsonElement, ILogger logger)
 		{
 			string jsonText = jsonElement.ToString();
 			if (field.IsInteger() && long.TryParse(jsonText, out long longValue))
 			{
+				logger.LogDebug("Field equality check type was {Type} value was {Value} and converted value was {ConvertedValue}", "integer", field.AsInt64(), longValue);
 				return field.AsInt64().Equals(longValue);
 			}
 			if (field.IsFloat() && float.TryParse(jsonText, out float floatValue))
 			{
+				logger.LogDebug("Field equality check type was {Type} value was {Value} and converted value was {ConvertedValue}", "float", field.AsFloat(), floatValue);
+
 				return field.AsFloat().Equals(floatValue);
 			} 
 			if (field.IsString())
 			{
+				logger.LogDebug("Field equality check type was {Type} value was {Value} and converted value was {ConvertedValue}", "string", field.AsString(), jsonText);
+
 				return field.AsString().Equals(jsonText, StringComparison.OrdinalIgnoreCase);
 			} 
 			if (field.IsDateTime() && DateTime.TryParse(jsonText, out DateTime dtValue))
 			{
+				logger.LogDebug("Field equality check type was {Type} value was {Value} and converted value was {ConvertedValue}", "datetime", field.AsDateTime(), dtValue);
+
 				return field.AsDateTime().Equals(dtValue);
 			}
 			if (field.IsTimeSpan() && TimeSpan.TryParse(jsonText, out TimeSpan tsValue))
 			{
+				logger.LogDebug("Field equality check type was {Type} value was {Value} and converted value was {ConvertedValue}", "timespan", field.AsTimeSpan(), tsValue);
+
 				return field.AsTimeSpan().Equals(tsValue);
 			}
 			if (field.IsHash() && IoHash.TryParse(jsonText, out IoHash hashValue))
 			{
+				logger.LogDebug("Field equality check type was {Type} value was {Value} and converted value was {ConvertedValue}", "iohash", field.AsHash(), hashValue);
+
 				return field.AsHash().Equals(hashValue);
 			}
 			if (field.IsUuid() && Guid.TryParse(jsonText, out Guid uuidValue))
 			{
+				logger.LogDebug("Field equality check type was {Type} value was {Value} and converted value was {ConvertedValue}", "uuid", field.AsUuid(), uuidValue);
+
 				return field.AsUuid().Equals(uuidValue);
 			}
 			if (field.IsObjectId() && CbObjectId.TryParse(jsonText, out CbObjectId objectId))
 			{
+				logger.LogDebug("Field equality check type was {Type} value was {Value} and converted value was {ConvertedValue}", "objectId", field.AsObjectId(), objectId);
+
 				return field.AsObjectId().Equals(objectId);
 			}
 			throw new NotImplementedException($"Unable to convert json type {jsonElement.ValueKind} to type that can be compared to field {field.Name} of type {field.TypeWithFlags}");
