Cloud DDC - Potential fix for bugs in searching when field is a integer, seems to be treated as a float which fails the match. Also added a bit more debug logging.

diff --git a/Engine/Source/Programs/UnrealCloudDDC/Jupiter/Implementation/Builds/SearchOp.cs b/Engine/Source/Programs/UnrealCloudDDC/Jupiter/Implementation/Builds/SearchOp.cs
--- a/Engine/Source/Programs/UnrealCloudDDC/Jupiter/Implementation/Builds/SearchOp.cs
+++ b/Engine/Source/Programs/UnrealCloudDDC/Jupiter/Implementation/Builds/SearchOp.cs
@@ -128,17 +128,17 @@ namespace Jupiter.Implementation.Builds
 					{
 						expectedType = dtType;
 						actualType = DateTime.Parse(s);
-					} 
-					else if (double.TryParse(rawText, out double doubleType))
-					{
-						expectedType = doubleType;
-						actualType = double.Parse(s);
 					}
 					else if (long.TryParse(rawText, out long longType))
 					{
 						expectedType = longType;
 						actualType = long.Parse(s);
 					}
+					else if (double.TryParse(rawText, out double doubleType))
+					{
+						expectedType = doubleType;
+						actualType = double.Parse(s);
+					}
 					else
 					{
 						throw new NotImplementedException($"Unable to guess compare type of string value \'{s}\', if you want to use a comparable op then make sure the original object uses a comparable type");
@@ -185,7 +185,7 @@ namespace Jupiter.Implementation.Builds
 			{
 				bool isEqual = FieldEquals(field, expectedField, logger);
 				
-				logger.LogDebug("Ran equality check between values from json: {Value0} and cb of type: {Value1} result was {Result}", expectedField, field.GetType().ToString(), isEqual);
+				logger.LogDebug("Ran equality check between values from json: {Value0} and cb of type: {Value1} result was {Result}", expectedField, field.Value?.ToString(), isEqual);
 				if (op == OpType.Equals)
 				{
 					return isEqual;
@@ -207,12 +207,6 @@ namespace Jupiter.Implementation.Builds
 				logger.LogDebug("Field equality check type was {Type} value was {Value} and converted value was {ConvertedValue}", "integer", field.AsInt64(), longValue);
 				return field.AsInt64().Equals(longValue);
 			}
-			if (field.IsFloat() && float.TryParse(jsonText, out float floatValue))
-			{
-				logger.LogDebug("Field equality check type was {Type} value was {Value} and converted value was {ConvertedValue}", "float", field.AsFloat(), floatValue);
-
-				return field.AsFloat().Equals(floatValue);
-			} 
 			if (field.IsString())
 			{
 				logger.LogDebug("Field equality check type was {Type} value was {Value} and converted value was {ConvertedValue}", "string", field.AsString(), jsonText);
@@ -249,6 +243,12 @@ namespace Jupiter.Implementation.Builds
 
 				return field.AsObjectId().Equals(objectId);
 			}
+			if (field.IsFloat() && float.TryParse(jsonText, out float floatValue))
+			{
+				logger.LogDebug("Field equality check type was {Type} value was {Value} and converted value was {ConvertedValue}", "float", field.AsFloat(), floatValue);
+
+				return field.AsFloat().Equals(floatValue);
+			}
 			throw new NotImplementedException($"Unable to convert json type {jsonElement.ValueKind} to type that can be compared to field {field.Name} of type {field.TypeWithFlags}");
 		}
 
