Cloud DDC - Added ability to tweak payload buffering sizes for tradeof between temp storage use and memory use esspecially usefull for the worker process which has very limited storage and memory.
Also fixed a bug were the filesystem GC setting was set incorrectly if the CacheRef settings were used on a worker (which is not a common scenario)

diff --git a/Engine/Source/Programs/UnrealCloudDDC/Helm/UnrealCloudDDC/templates/configmap-worker.yaml b/Engine/Source/Programs/UnrealCloudDDC/Helm/UnrealCloudDDC/templates/configmap-worker.yaml
--- a/Engine/Source/Programs/UnrealCloudDDC/Helm/UnrealCloudDDC/templates/configmap-worker.yaml
+++ b/Engine/Source/Programs/UnrealCloudDDC/Helm/UnrealCloudDDC/templates/configmap-worker.yaml
@@ -108,15 +108,15 @@ ConsistencyCheck:
 {{ $config.ConsistencyCheck | toYaml | nindent 2 }}
 {{- end }}
 
-{{- if $config.GC }}
-GC: 
-{{ $config.GC | toYaml | nindent 2 }}
-
 {{- if $config.CacheRef }}
 CacheRef:
 {{ $config.CacheRef | toYaml | nindent 2 }}
 {{- end }}
 
+{{- if $config.GC }}
+GC: 
+{{ $config.GC | toYaml | nindent 2 }}
+
 {{- if or ($worker.persistence.enabled) ($worker.persistence.volume) }}
   RunFilesystemCleanup: true
 {{- end }}
@@ -127,6 +127,11 @@ Snapshot:
 {{ $config.Snapshot | toYaml | nindent 2 }}
 {{- end }}
 
+{{- if $config.PayloadBuffering }}
+PayloadBuffering: 
+{{ $config.PayloadBuffering | toYaml | nindent 2 }}
+{{- end }}
+
 {{- if and (.Values.global.cloudProvider) (eq .Values.global.cloudProvider "Azure") }}
 {{ include "cloud.settings.azure" .Values.worker }}
 {{- else if and (.Values.global.cloudProvider) (eq .Values.global.cloudProvider "AWS") }}
diff --git a/Engine/Source/Programs/UnrealCloudDDC/Helm/UnrealCloudDDC/templates/configmap.yaml b/Engine/Source/Programs/UnrealCloudDDC/Helm/UnrealCloudDDC/templates/configmap.yaml
--- a/Engine/Source/Programs/UnrealCloudDDC/Helm/UnrealCloudDDC/templates/configmap.yaml
+++ b/Engine/Source/Programs/UnrealCloudDDC/Helm/UnrealCloudDDC/templates/configmap.yaml
@@ -157,6 +157,11 @@ Scylla:
 {{ $config.Scylla | toYaml | nindent 2 }}
 {{- end }}
 
+{{- if $config.PayloadBuffering }}
+PayloadBuffering: 
+{{ $config.PayloadBuffering | toYaml | nindent 2 }}
+{{- end }}
+
 {{- if $config.Metrics }}
 Metrics: 
 {{ $config.Metrics | toYaml | nindent 2 }}
diff --git a/Engine/Source/Programs/UnrealCloudDDC/Jupiter.Common/BaseStartup.cs b/Engine/Source/Programs/UnrealCloudDDC/Jupiter.Common/BaseStartup.cs
--- a/Engine/Source/Programs/UnrealCloudDDC/Jupiter.Common/BaseStartup.cs
+++ b/Engine/Source/Programs/UnrealCloudDDC/Jupiter.Common/BaseStartup.cs
@@ -92,7 +92,6 @@ namespace Jupiter
 			services.AddOptions<ServiceAccountAuthOptions>().Bind(Configuration.GetSection("ServiceAccounts")).ValidateDataAnnotations();
 
 			services.AddOptions<JupiterSettings>().Bind(Configuration.GetSection("Jupiter")).ValidateDataAnnotations();
-			services.AddOptions<BufferedPayloadOptions>().Bind(Configuration.GetSection("Jupiter")).ValidateDataAnnotations();
 			services.AddOptions<NamespaceSettings>().Bind(Configuration.GetSection("Namespaces")).ValidateDataAnnotations();
 
 			services.AddSingleton(typeof(INamespacePolicyResolver), typeof(NamespacePolicyResolver));
