Cloud DDC - Added more tracing to debug issue with blob replication

diff --git a/Engine/Source/Programs/UnrealCloudDDC/Jupiter/Implementation/Blob/AmazonS3Store.cs b/Engine/Source/Programs/UnrealCloudDDC/Jupiter/Implementation/Blob/AmazonS3Store.cs
--- a/Engine/Source/Programs/UnrealCloudDDC/Jupiter/Implementation/Blob/AmazonS3Store.cs
+++ b/Engine/Source/Programs/UnrealCloudDDC/Jupiter/Implementation/Blob/AmazonS3Store.cs
@@ -421,6 +421,7 @@ namespace Jupiter.Implementation
 				{
 					filePath = fileStream.Name;
 				}
+				// use multipart transfers when buffer is larger than 16 MB, which is also the S3 default
 				else if (stream.Length > 16 * (long)Math.Pow(2, 20))
 				{
 					// will be chunked by TransferUtility
diff --git a/Engine/Source/Programs/UnrealCloudDDC/Jupiter/Implementation/Blob/BlobService.cs b/Engine/Source/Programs/UnrealCloudDDC/Jupiter/Implementation/Blob/BlobService.cs
--- a/Engine/Source/Programs/UnrealCloudDDC/Jupiter/Implementation/Blob/BlobService.cs
+++ b/Engine/Source/Programs/UnrealCloudDDC/Jupiter/Implementation/Blob/BlobService.cs
@@ -296,7 +296,7 @@ public class BlobService : IBlobService
 					BlobContents fallbackContent = await GetObjectFromStoresAsync(policy.FallbackNamespace.Value, blob, storageLayers, cancellationToken: cancellationToken);
 
 					// populate the primary namespace with the content
-					using TelemetrySpan _ = _tracer.StartActiveSpan("HierarchicalStore.Populate").SetAttribute("operation.name", "HierarchicalStore.Populate");
+					using TelemetrySpan _ = _tracer.StartActiveSpan("HierarchicalStore.Populate").SetAttribute("operation.name", "HierarchicalStore.Populate").SetAttribute("resource.name", blob.ToString());
 					using ServerTimingMetricScoped? serverTimingScope = serverTiming?.CreateServerTimingMetricScope($"blob.populate", "Populating caches with blob contents");
 
 					await using MemoryStream tempStream = new MemoryStream();
@@ -585,7 +585,7 @@ public class BlobService : IBlobService
 			throw new NotSupportedException($"Replication is not allowed in namespace {ns}");
 		}
 		IServerTiming? serverTiming = _httpContextAccessor.HttpContext?.RequestServices.GetService<IServerTiming>();
-		using TelemetrySpan scope = _tracer.StartActiveSpan("HierarchicalStore.Replicate").SetAttribute("operation.name", "HierarchicalStore.Replicate");
+		using TelemetrySpan scope = _tracer.StartActiveSpan("HierarchicalStore.Replicate").SetAttribute("operation.name", "HierarchicalStore.Replicate").SetAttribute("resource.name", blob.ToString());
 
 		using ServerTimingMetricScoped? serverTimingScope = serverTiming?.CreateServerTimingMetricScope("blob.replicate", "Replicating blob from remote instances");
 
