Cloud DDC - More tracing around the on demand replication to track issue with block metadata

diff --git a/Engine/Source/Programs/UnrealCloudDDC/Jupiter/Implementation/Blob/BlobService.cs b/Engine/Source/Programs/UnrealCloudDDC/Jupiter/Implementation/Blob/BlobService.cs
--- a/Engine/Source/Programs/UnrealCloudDDC/Jupiter/Implementation/Blob/BlobService.cs
+++ b/Engine/Source/Programs/UnrealCloudDDC/Jupiter/Implementation/Blob/BlobService.cs
@@ -625,11 +625,15 @@ public class BlobService : IBlobService
 			{
 				url += "/";
 			}
+			using TelemetrySpan replicateScope = _tracer.StartActiveSpan("HierarchicalStore.Replicate.From").SetAttribute("operation.name", "HierarchicalStore.Replicate.From").SetAttribute("resource.name", $"{region}-{blob}");
+
 			using HttpRequestMessage blobRequest = await BuildHttpRequestAsync(HttpMethod.Get, new Uri($"{url}api/v1/blobs/{ns}/{blob}?allowOndemandReplication=false"));
 			HttpResponseMessage blobResponse = await httpClient.SendAsync(blobRequest, HttpCompletionOption.ResponseHeadersRead, cancellationToken);
 
 			if (blobResponse.StatusCode == HttpStatusCode.NotFound)
 			{
+				replicateScope.SetAttribute("found", false);
+
 				// try the next region
 				continue;
 			}
@@ -637,8 +641,11 @@ public class BlobService : IBlobService
 			if (blobResponse.StatusCode != HttpStatusCode.OK)
 			{
 				_logger.LogWarning("Failed to replicate {Blob} in {Namespace} from region {Region} due to bad http status code {StatusCode}.", blob, ns, region, blobResponse.StatusCode);
-				throw new BlobReplicationException(ns, blob, $"Failed to replicate {blob} in {ns} from region {region} due to bad http status code {blobResponse.StatusCode}.");
+				BlobReplicationException e = new BlobReplicationException(ns, blob, $"Failed to replicate {blob} in {ns} from region {region} due to bad http status code {blobResponse.StatusCode}.");
+				replicateScope.RecordException(e);
+				throw e;
 			}
+			replicateScope.SetAttribute("found", true);
 
 			await using Stream s = await blobResponse.Content.ReadAsStreamAsync(cancellationToken);
 			using IBufferedPayload payload = await _bufferedPayloadFactory.CreateFilesystemBufferedPayloadAsync(s, "replicate-blob", cancellationToken);
