Cloud DDC - Removed redundant payload buffering options

diff --git a/Engine/Source/Programs/UnrealCloudDDC/Helm/UnrealCloudDDC/templates/configmap.yaml b/Engine/Source/Programs/UnrealCloudDDC/Helm/UnrealCloudDDC/templates/configmap.yaml
--- a/Engine/Source/Programs/UnrealCloudDDC/Helm/UnrealCloudDDC/templates/configmap.yaml
+++ b/Engine/Source/Programs/UnrealCloudDDC/Helm/UnrealCloudDDC/templates/configmap.yaml
@@ -157,11 +157,6 @@ Scylla:
 {{ $config.Scylla | toYaml | nindent 2 }}
 {{- end }}
 
-{{- if $config.PayloadBuffering }}
-PayloadBuffering: 
-{{ $config.PayloadBuffering | toYaml | nindent 2 }}
-{{- end }}
-
 {{- if $config.Metrics }}
 Metrics: 
 {{ $config.Metrics | toYaml | nindent 2 }}
