Cloud DDC - Added more trace information when ondemand replicating

diff --git a/Engine/Source/Programs/UnrealCloudDDC/Jupiter/Implementation/Blob/BlobService.cs b/Engine/Source/Programs/UnrealCloudDDC/Jupiter/Implementation/Blob/BlobService.cs
--- a/Engine/Source/Programs/UnrealCloudDDC/Jupiter/Implementation/Blob/BlobService.cs
+++ b/Engine/Source/Programs/UnrealCloudDDC/Jupiter/Implementation/Blob/BlobService.cs
@@ -601,8 +601,10 @@ public class BlobService : IBlobService
 
 		if (!regions.Any())
 		{
+			scope.SetAttribute("found", false);
 			throw new BlobReplicationException(ns, blob, "Blob not found in any region");
 		}
+		scope.SetAttribute("found", true);
 
 		_logger.LogInformation("On-demand replicating blob {Blob} in Namespace {Namespace}", blob, ns);
 		List<(int, string)> possiblePeers = new List<(int, string)>(_peerStatusService.GetPeersByLatency(regions));
@@ -654,13 +656,14 @@ public class BlobService : IBlobService
 			break;
 		}
 
+		scope.SetAttribute("replicated", replicated);
+
 		if (!replicated)
 		{
 			_logger.LogWarning("Failed to replicate {Blob} in {Namespace} due to it not existing in any region", blob, ns);
 
 			throw new BlobReplicationException(ns, blob, $"Failed to replicate {blob} in {ns} due to it not existing in any region");
 		}
-
 		return await GetObjectAsync(ns, blob, supportsRedirectUri: supportsRedirectUri, bucketHint: bucketHint, cancellationToken: cancellationToken);
 	}
 
