Cloud DDC - Reconfigured forwarding headers to allow for more then one forwarding hop and removed assumption on proxy being on localhost.
Should fix issues with source ip not always being visisble in the connection.

diff --git a/Engine/Source/Programs/UnrealCloudDDC/Jupiter.Common/BaseStartup.cs b/Engine/Source/Programs/UnrealCloudDDC/Jupiter.Common/BaseStartup.cs
--- a/Engine/Source/Programs/UnrealCloudDDC/Jupiter.Common/BaseStartup.cs
+++ b/Engine/Source/Programs/UnrealCloudDDC/Jupiter.Common/BaseStartup.cs
@@ -389,6 +389,19 @@ namespace Jupiter
 				opt.IncludeFormattedMessage = true;
 			});
 
+			services.Configure<ForwardedHeadersOptions>(options =>
+			{
+				// include enough ips to allow for multiple proxies in front
+				options.ForwardLimit = 3;
+
+				// we do not know the address of load balancers in front of us but want to accept the forwarding headers from them so we reset these lists
+				// we generally assume that traffic happens over https and that the network we are on is very limited in scope (only ingress via load balancers)
+				// TODO: If we wanted to we could introduce a setting to specify the ip range of the load balancers here, but would be kind of annoying to have to specify
+				// , but it could be useful if this is run in less secure networks
+				options.KnownNetworks.Clear();
+				options.KnownProxies.Clear();
+			});
+
 			services.AddSingleton<Tracer>(CreateTracer);
 
 			services.AddSingleton<Meter>(CreateMeter);
