Cloud DDC - Added checks to see how many empty namespaces and buckets are found during garbage collection of refs. Considering cleaning up the tracking of these using this information but would want to make sure it would only operate on namespaces that have no content in them.

diff --git a/Engine/Source/Programs/UnrealCloudDDC/Jupiter/Implementation/GC/CacheCleanup.cs b/Engine/Source/Programs/UnrealCloudDDC/Jupiter/Implementation/GC/CacheCleanup.cs
--- a/Engine/Source/Programs/UnrealCloudDDC/Jupiter/Implementation/GC/CacheCleanup.cs
+++ b/Engine/Source/Programs/UnrealCloudDDC/Jupiter/Implementation/GC/CacheCleanup.cs
@@ -1,6 +1,7 @@
 ﻿// Copyright Epic Games, Inc. All Rights Reserved.
 
 using System;
+using System.Collections.Concurrent;
 using System.Collections.Generic;
 using System.Diagnostics.Metrics;
 using System.Threading;
@@ -95,6 +96,8 @@ namespace Jupiter.Implementation
 			long consideredCount = 0;
 			DateTime cleanupStart = DateTime.Now;
 
+			ConcurrentDictionary<NamespaceId, ConcurrentDictionary<BucketId, bool>> namespacesWithRefs = new();
+
 			await Parallel.ForEachAsync(_referencesStore.GetRecordsAsync(cancellationToken),
 				new ParallelOptions
 				{
@@ -104,6 +107,10 @@ namespace Jupiter.Implementation
 				{
 					(NamespaceId ns, BucketId bucket, RefId name, DateTime lastAccessTime) = tuple;
 
+					// keep track of which namespaces and buckets we have seen
+					ConcurrentDictionary<BucketId, bool> bucketTracking = namespacesWithRefs.GetOrAdd(ns, new ConcurrentDictionary<BucketId, bool>());
+					bucketTracking.TryAdd(bucket, true);
+
 					if (!ShouldGCNamespace(ns))
 					{
 						return;
@@ -154,6 +161,30 @@ namespace Jupiter.Implementation
 			_logger.LogInformation(
 				"Finished cleaning refs. Refs considered: {ConsideredCount} Refs Deleted: {DeletedCount}. Cleanup took: {CleanupDuration}", consideredCount, countOfDeletedRecords, cleanupDuration);
 
+			await foreach (NamespaceId ns in _referencesStore.GetNamespacesAsync(cancellationToken))
+			{
+				if (namespacesWithRefs.ContainsKey(ns))
+				{
+					continue;
+				}
+
+				ConcurrentDictionary<BucketId, bool> foundBuckets = namespacesWithRefs[ns];
+				await foreach (BucketId bucket in _referencesStore.GetBucketsAsync(ns, cancellationToken))
+				{
+					// if the bucket had a ref we should not remove its tracking
+					if (foundBuckets.ContainsKey(bucket))
+					{
+						continue;
+					}
+					_logger.LogWarning("Namespace {Namespace} Bucket {Bucket} has no refs. Should be removed from tracking.", ns, bucket);
+					// TODO: We need a ligher way to drop a bucket that has no refs
+					// await _referencesStore.DeleteBucketAsync(ns, bucket, cancellationToken);
+				}
+				_logger.LogWarning("Namespace {Namespace} has no refs. Should be removed from tracking.", ns);
+
+				// TODO: We should have a lighter namespace drop that does not remove refs, option to it?
+				//await _referencesStore.DropNamespaceAsync(ns, cancellationToken);
+			}
 			return countOfDeletedRecords;
 		}
 
