Cloud DDC - Changed the auth endpoint to use the new method for controlling auth to run its checks which also includes debug logging that explains the results.

diff --git a/Engine/Source/Programs/UnrealCloudDDC/Jupiter/Controllers/AuthController.cs b/Engine/Source/Programs/UnrealCloudDDC/Jupiter/Controllers/AuthController.cs
--- a/Engine/Source/Programs/UnrealCloudDDC/Jupiter/Controllers/AuthController.cs
+++ b/Engine/Source/Programs/UnrealCloudDDC/Jupiter/Controllers/AuthController.cs
@@ -44,23 +44,24 @@ namespace Jupiter.Controllers
 		[HttpGet("{ns}/actions")]
 		public IActionResult Actions([FromRoute][Required] NamespaceId ns)
 		{
-			NamespacePolicy policy = _namespacePolicyResolver.GetPoliciesForNs(ns);
-
 			List<JupiterAclAction> allowedActions = new List<JupiterAclAction>();
-			foreach (AclEntry acl in policy.Acls)
+			
+			CaptureLogger captureLogger = new CaptureLogger();
+			bool policiesFound = false;
+			foreach (AclPolicy acl in _authSettings.CurrentValue.Policies)
 			{
-				allowedActions.AddRange(acl.Resolve(User));
+				policiesFound = true;
+				allowedActions.AddRange(acl.Resolve(User, new AccessScope(ns), captureLogger));
 			}
 
-			// the root and namespace acls are combined, namespace acls can not override what we define in the root
-			foreach (AclEntry acl in _authSettings.CurrentValue.Acls)
+			if (!policiesFound)
 			{
-				allowedActions.AddRange(acl.Resolve(User));
+				captureLogger.LogWarning("No policies set so no actions generated");
+				return BadRequest(new JsonResult(new { Actions = allowedActions, LogOutput = captureLogger.RenderLines() }));
 			}
 
-			return Ok(new JsonResult(new { Actions = allowedActions }));
+			return Ok(new JsonResult(new { Actions = allowedActions , LogOutput = captureLogger.RenderLines()}));
 		}
-
 		
 		[HttpGet("{ns}/{bucket}/actions")]
 		public IActionResult Actions([FromRoute][Required] NamespaceId ns, [FromRoute][Required] BucketId bucket)
