Cloud DDC - Added more debug logging in the acl matching to try and catch issue with namespace global operations not being correctly allowed.

diff --git a/Engine/Source/Programs/UnrealCloudDDC/Jupiter.Common/Authentication/Acl.cs b/Engine/Source/Programs/UnrealCloudDDC/Jupiter.Common/Authentication/Acl.cs
--- a/Engine/Source/Programs/UnrealCloudDDC/Jupiter.Common/Authentication/Acl.cs
+++ b/Engine/Source/Programs/UnrealCloudDDC/Jupiter.Common/Authentication/Acl.cs
@@ -293,7 +293,7 @@ public class AclScopeEntry
 				}
 			}
 
-			Match match = Regex!.Match(scope);
+			Match match = Regex.Match(scope);
 			if (!match.Success)
 			{
 				logger.LogDebug("Regular expression: \"{Regex}\", did not match scope \"{Scope}\" for policy {Name}", Match, scope, policyName);
@@ -325,7 +325,7 @@ public class AclScopeEntry
 			else
 			{
 				matches = value.Equals(expectedScopeValue, StringComparison.OrdinalIgnoreCase);
-
+				logger.LogDebug("Compared values {Value0} and {Value1} with match result {Result} for policy {Name}. ", value, expectedScopeValue, matches, policyName);
 			}
 
 			if (Not && matches)
@@ -347,13 +347,13 @@ public class AclScopeEntry
 
 		if (Not)
 		{
-			logger.LogDebug("No values matched scope \"{Scope}\" in policy {Name}. This was a inverted operation so is considered a match.", scope, policyName);
+			logger.LogDebug("No values matched scope \"{Scope}\" in policy {Name}. This was a inverted operation so is considered a match.", expectedScopeValue, policyName);
 
 			// no match found, that means this policy applies if the Not flag is set
 			return true;
 		}
 
-		logger.LogDebug("No values matched scope \"{Scope}\" in policy {Name}.", scope, policyName);
+		logger.LogDebug("No values matched scope \"{Scope}\" in policy {Name}.", expectedScopeValue, policyName);
 		return false;
 	}
 }
