Cloud DDC - For operations that do not run in a bucket as long as any profile grants access to the namespace that operation will be allowed.

diff --git a/Engine/Source/Programs/UnrealCloudDDC/Jupiter.Common/Authentication/Acl.cs b/Engine/Source/Programs/UnrealCloudDDC/Jupiter.Common/Authentication/Acl.cs
--- a/Engine/Source/Programs/UnrealCloudDDC/Jupiter.Common/Authentication/Acl.cs
+++ b/Engine/Source/Programs/UnrealCloudDDC/Jupiter.Common/Authentication/Acl.cs
@@ -212,6 +212,12 @@ public class AclPolicy
 			return Array.Empty<JupiterAclAction>();
 		}
 
+		if (scope.Bucket == null)
+		{
+			// this is an operation that doesn't run on a bucket, if we have access to any bucket in the namespace we are allowed to do these operations as we are allowed to be aware of this namespace existing
+			return Actions;
+		}
+
 		bool bucketFound = false;
 		foreach (AclScopeEntry scopeBuckets in Scope.Buckets)
 		{
