Cloud DDC - Fixed issue were block metadata references was incorrect, it should model the incoming dependency (metadata can not be deleted until its block doesn't exist anymore) but it was modelling the outgoing dependency which did not match with how it is requested.

diff --git a/Engine/Source/Programs/UnrealCloudDDC/Jupiter.Common/BaseStartup.cs b/Engine/Source/Programs/UnrealCloudDDC/Jupiter.Common/BaseStartup.cs
--- a/Engine/Source/Programs/UnrealCloudDDC/Jupiter.Common/BaseStartup.cs
+++ b/Engine/Source/Programs/UnrealCloudDDC/Jupiter.Common/BaseStartup.cs
@@ -16,7 +16,6 @@ using System.Threading.Tasks;
 using Amazon;
 using EpicGames.AspNet;
 using Jupiter.Common;
-using Jupiter.Common.Implementation;
 using Microsoft.AspNetCore.Authentication;
 using Microsoft.AspNetCore.Authentication.JwtBearer;
 using Microsoft.AspNetCore.Authorization;
diff --git a/Engine/Source/Programs/UnrealCloudDDC/Jupiter/Controllers/AuthController.cs b/Engine/Source/Programs/UnrealCloudDDC/Jupiter/Controllers/AuthController.cs
--- a/Engine/Source/Programs/UnrealCloudDDC/Jupiter/Controllers/AuthController.cs
+++ b/Engine/Source/Programs/UnrealCloudDDC/Jupiter/Controllers/AuthController.cs
@@ -5,7 +5,6 @@ using System.ComponentModel.DataAnnotations;
 using System.Threading.Tasks;
 using EpicGames.Core;
 using EpicGames.Horde.Storage;
-using Jupiter.Common;
 using Microsoft.AspNetCore.Authorization;
 using Microsoft.AspNetCore.Mvc;
 using Microsoft.Extensions.Logging;
@@ -19,13 +18,11 @@ namespace Jupiter.Controllers
 	public class AuthController : ControllerBase
 	{
 		private readonly IRequestHelper _requestHelper;
-		private readonly INamespacePolicyResolver _namespacePolicyResolver;
 		private readonly IOptionsMonitor<AuthSettings> _authSettings;
 
-		public AuthController(IRequestHelper requestHelper, INamespacePolicyResolver namespacePolicyResolver, IOptionsMonitor<AuthSettings> authSettings)
+		public AuthController(IRequestHelper requestHelper, IOptionsMonitor<AuthSettings> authSettings)
 		{
 			_requestHelper = requestHelper;
-			_namespacePolicyResolver = namespacePolicyResolver;
 			_authSettings = authSettings;
 		}
 
diff --git a/Engine/Source/Programs/UnrealCloudDDC/Jupiter/Controllers/BuildsController.cs b/Engine/Source/Programs/UnrealCloudDDC/Jupiter/Controllers/BuildsController.cs
--- a/Engine/Source/Programs/UnrealCloudDDC/Jupiter/Controllers/BuildsController.cs
+++ b/Engine/Source/Programs/UnrealCloudDDC/Jupiter/Controllers/BuildsController.cs
@@ -1093,7 +1093,9 @@ namespace Jupiter.Controllers
 
 				Tracer.CurrentSpan.SetAttribute("blockId", blockId.ToString());
 				await _blockStore.PutBlockMetadataAsync(ns, blockId, metadataBlockId);
-				await _blobIndex.AddBlobReferencesAsync(ns, blockId, metadataBlockId);
+
+				// blob references track incoming references e.g. the block requires the metadata block to be kept around
+				await _blobIndex.AddBlobReferencesAsync(ns, metadataBlockId, blockId);
 
 				return Ok();
 			}
diff --git a/Engine/Source/Programs/UnrealCloudDDC/Jupiter/Implementation/Blob/IBlobIndex.cs b/Engine/Source/Programs/UnrealCloudDDC/Jupiter/Implementation/Blob/IBlobIndex.cs
--- a/Engine/Source/Programs/UnrealCloudDDC/Jupiter/Implementation/Blob/IBlobIndex.cs
+++ b/Engine/Source/Programs/UnrealCloudDDC/Jupiter/Implementation/Blob/IBlobIndex.cs
@@ -33,7 +33,7 @@ namespace Jupiter.Implementation.Blob
 
 		Task RemoveReferencesAsync(NamespaceId ns, BlobId id, List<BaseBlobReference>? referencesToRemove, CancellationToken cancellationToken = default);
 		Task<List<string>> GetBlobRegionsAsync(NamespaceId ns, BlobId blob, CancellationToken cancellationToken = default);
-		Task AddBlobReferencesAsync(NamespaceId ns, BlobId sourceBlob, BlobId targetBlob, CancellationToken cancellationToken = default);
+		Task AddBlobReferencesAsync(NamespaceId ns, BlobId blob, BlobId blobThatReferences, CancellationToken cancellationToken = default);
 
 		Task AddBlobToBucketListAsync(NamespaceId ns, BucketId bucket, RefId key, BlobId blobId, long blobSize, CancellationToken cancellationToken = default);
 		Task RemoveBlobFromBucketListAsync(NamespaceId ns, BucketId bucket, RefId key, List<BlobId> blobIds, CancellationToken cancellationToken = default);
