Cloud DDC - Added warning when we delete a somewhat new blob as that is an indication some tracking is incorrect (though not guaranteed as adding blobs without references is allowed)

diff --git a/Engine/Source/Programs/UnrealCloudDDC/Jupiter/Implementation/GC/OrphanBlobCleanupRefs.cs b/Engine/Source/Programs/UnrealCloudDDC/Jupiter/Implementation/GC/OrphanBlobCleanupRefs.cs
--- a/Engine/Source/Programs/UnrealCloudDDC/Jupiter/Implementation/GC/OrphanBlobCleanupRefs.cs
+++ b/Engine/Source/Programs/UnrealCloudDDC/Jupiter/Implementation/GC/OrphanBlobCleanupRefs.cs
@@ -224,6 +224,12 @@ namespace Jupiter.Implementation
 			_deletedBlobCounter.Add(1);
 			_logger.LogInformation("GC Orphan blob {Blob} from {StoragePool} which was last modified at {LastModifiedTime}", blob, storagePoolName, lastModifiedTime);
 
+			if (lastModifiedTime > DateTime.Now.AddHours(-6.0))
+			{
+				// the blob was added within the last 6 hours and then deleted, emit warning as that indicates the blob likely wasn't needed by anything. This warning can help detect bugs were blobs that should be referenced are getting deleted for some reason
+				_logger.LogWarning("Recently added blob {Blob} from {StoragePool} was deleted. Indicates blob upload without anything that referenced it which is safe but can be a sign of an issue", blob, storagePoolName);
+			}
+
 			// delete block information if this blob is a block
 			foreach (NamespaceId namespaceId in namespacesThatSharePool)
 			{
