Cloud DDC - Added audit log record when using the search api about how many hits people had, lets us find people that potentially try to enumerate all contents for a lot of branches.

diff --git a/Engine/Source/Programs/UnrealCloudDDC/Jupiter/Controllers/BuildsController.cs b/Engine/Source/Programs/UnrealCloudDDC/Jupiter/Controllers/BuildsController.cs
--- a/Engine/Source/Programs/UnrealCloudDDC/Jupiter/Controllers/BuildsController.cs
+++ b/Engine/Source/Programs/UnrealCloudDDC/Jupiter/Controllers/BuildsController.cs
@@ -398,6 +398,9 @@ namespace Jupiter.Controllers
 			}
 			_logger.LogDebug("Search completed in {Namespace} {Bucket} request: {@SearchRequest} Results: {CountOfResults} Considered: {CountOfBuildsConsidered} took {Duration}", ns, bucket, searchRequest, foundResults.Count, count, timeSpent);
 			Tracer.CurrentSpan.SetAttribute("countOfMatches", foundResults.Count);
+
+			_auditLogger?.Information("{HttpMethod} {Endpoint} IP:{IP} User:{Username} UserAgent:\"{Useragent}\" MatchCount:{MatchCount}", HttpMethod.Post, nameof(SearchBuildAsync), Request.HttpContext.Connection.RemoteIpAddress, User?.Identity?.Name ?? "Unknown-user", string.Join(' ', Request.Headers.UserAgent.ToArray()), foundResults.Count);
+
 			return Ok(new SearchResult(foundResults, partialResult));
 		}
 
