Cloud DDC - Made the audit log more consistent between endpoints (mostly makes sure namespace and bucket is always present when appopriate).
Also added the ue session id and the jti (jwt id) for detection of suspiciuous clients.

diff --git a/Engine/Source/Programs/UnrealCloudDDC/Jupiter/Controllers/BuildsController.cs b/Engine/Source/Programs/UnrealCloudDDC/Jupiter/Controllers/BuildsController.cs
--- a/Engine/Source/Programs/UnrealCloudDDC/Jupiter/Controllers/BuildsController.cs
+++ b/Engine/Source/Programs/UnrealCloudDDC/Jupiter/Controllers/BuildsController.cs
@@ -398,8 +398,9 @@ namespace Jupiter.Controllers
 			}
 			_logger.LogDebug("Search completed in {Namespace} {Bucket} request: {@SearchRequest} Results: {CountOfResults} Considered: {CountOfBuildsConsidered} took {Duration}", ns, bucket, searchRequest, foundResults.Count, count, timeSpent);
 			Tracer.CurrentSpan.SetAttribute("countOfMatches", foundResults.Count);
+			_diagnosticContext.Set("MatchCount", foundResults.Count);
 
-			_auditLogger?.Information("{HttpMethod} {Endpoint} IP:{IP} User:{Username} UserAgent:\"{Useragent}\" MatchCount:{MatchCount}", HttpMethod.Post, nameof(SearchBuildAsync), Request.HttpContext.Connection.RemoteIpAddress, User?.Identity?.Name ?? "Unknown-user", string.Join(' ', Request.Headers.UserAgent.ToArray()), foundResults.Count);
+			LogAuditEntry(HttpMethod.Post, nameof(SearchBuildAsync), ns, bucket, null, null);
 
 			return Ok(new SearchResult(foundResults, partialResult));
 		}
@@ -1461,7 +1462,8 @@ namespace Jupiter.Controllers
 			{
 				return new ForbidResult();
 			}
-			_auditLogger?.Information("{HttpMethod} {Endpoint} IP:{IP} User:{Username} UserAgent:\"{Useragent}\"", HttpMethod.Get, nameof(GetNamespacesAsync), Request.HttpContext.Connection.RemoteIpAddress, User?.Identity?.Name ?? "Unknown-user", string.Join(' ', Request.Headers.UserAgent.ToArray()));
+
+			LogAuditEntry(HttpMethod.Get, nameof(GetNamespacesAsync), null, null, null, null);
 
 			return Ok(new GetNamespacesResponse(namespacesWithAccess.ToArray()));
 		}
@@ -1470,7 +1472,7 @@ namespace Jupiter.Controllers
 		[ProducesDefaultResponseType]
 		public async Task<IActionResult> GetBucketsAsync(NamespaceId ns)
 		{
-			_auditLogger?.Information("{HttpMethod} {Endpoint} {Namespace} IP:{IP} User:{Username} UserAgent:\"{Useragent}\"", HttpMethod.Get, nameof(GetBucketsAsync), ns, Request.HttpContext.Connection.RemoteIpAddress, User.Identity?.Name ?? "Unknown-user", string.Join(' ', Request.Headers.UserAgent.ToArray()));
+			LogAuditEntry(HttpMethod.Get, nameof(GetBucketsAsync), ns, null, null, null);
 
 			BucketId[] potentialBuckets = await _refService.GetBucketsAsync(ns, HttpContext.RequestAborted).ToArrayAsync(HttpContext.RequestAborted);
 
@@ -1494,9 +1496,35 @@ namespace Jupiter.Controllers
 
 		#endregion
 
-		private void LogAuditEntry(HttpMethod method, string endpoint, NamespaceId ns, BucketId bucketId, CbObjectId buildId, string? partId)
+		private void LogAuditEntry(HttpMethod method, string endpoint, NamespaceId? ns, BucketId? bucketId, CbObjectId? buildId, string? partId)
 		{
-			_auditLogger?.Information("{HttpMethod} {Endpoint} '{Namespace}/{BucketId}:{BuildId}' Part:{PartId} IP:{IP} User:{Username} UserAgent:\"{Useragent}\"", method, endpoint, ns, bucketId, buildId, partId, Request.HttpContext.Connection.RemoteIpAddress, User?.Identity?.Name ?? "Unknown-user", string.Join(' ', Request.Headers.UserAgent.ToArray()));
+			string sessionHeader = "None";
+			if (Request.Headers.TryGetValue("ue-session", out StringValues sessionHeaderValues))
+			{
+				sessionHeader = sessionHeaderValues.First() ?? "None";
+			}
+
+			string? jti = User.Claims.FirstOrDefault(claim => claim.Type == "jti")?.Value ?? null;
+
+			if (ns == null)
+			{
+				_auditLogger?.Information("{HttpMethod} {Endpoint} IP:{IP} User:{Username} UserAgent:\"{Useragent}\" SessionId: {SessionId} Jti:{Jti}", method, endpoint, Request.HttpContext.Connection.RemoteIpAddress, User?.Identity?.Name ?? "Unknown-user", string.Join(' ', Request.Headers.UserAgent.ToArray()), sessionHeader, jti);
+				return;
+			}
+
+			if (bucketId == null)
+			{
+				_auditLogger?.Information("{HttpMethod} {Endpoint} '{Namespace}' IP:{IP} User:{Username} UserAgent:\"{Useragent}\" SessionId: {SessionId} Jti:{Jti}", method, endpoint, ns, Request.HttpContext.Connection.RemoteIpAddress, User?.Identity?.Name ?? "Unknown-user", string.Join(' ', Request.Headers.UserAgent.ToArray()), sessionHeader, jti);
+				return;
+			}
+
+			if (partId == null)
+			{
+				_auditLogger?.Information("{HttpMethod} {Endpoint} '{Namespace}/{BucketId}:{BuildId}' IP:{IP} User:{Username} UserAgent:\"{Useragent}\" SessionId: {SessionId} Jti:{Jti}", method, endpoint, ns, bucketId, buildId, Request.HttpContext.Connection.RemoteIpAddress, User?.Identity?.Name ?? "Unknown-user", string.Join(' ', Request.Headers.UserAgent.ToArray()), sessionHeader, jti);
+				return;
+			}
+			
+			_auditLogger?.Information("{HttpMethod} {Endpoint} '{Namespace}/{BucketId}:{BuildId}' Part:{PartId} IP:{IP} User:{Username} UserAgent:\"{Useragent}\" SessionId: {SessionId} Jti:{Jti}", method, endpoint, ns, bucketId, buildId, partId, Request.HttpContext.Connection.RemoteIpAddress, User?.Identity?.Name ?? "Unknown-user", string.Join(' ', Request.Headers.UserAgent.ToArray()), sessionHeader, jti);
 		}
 	}
 
