Cloud DDC - Avoid buffering payloads during a cache miss if the stream is seekable (so we can upload it to all layers that are missing it)

diff --git a/Engine/Source/Programs/UnrealCloudDDC/Jupiter/Implementation/Blob/BlobService.cs b/Engine/Source/Programs/UnrealCloudDDC/Jupiter/Implementation/Blob/BlobService.cs
--- a/Engine/Source/Programs/UnrealCloudDDC/Jupiter/Implementation/Blob/BlobService.cs
+++ b/Engine/Source/Programs/UnrealCloudDDC/Jupiter/Implementation/Blob/BlobService.cs
@@ -429,8 +429,14 @@ public class BlobService : IBlobService
 			using TelemetrySpan _ = _tracer.StartActiveSpan("HierarchicalStore.Populate").SetAttribute("operation.name", "HierarchicalStore.Populate");
 			using ServerTimingMetricScoped? serverTimingScope = serverTiming?.CreateServerTimingMetricScope($"blob.populate", "Populating caches with blob contents");
 
+			bool streamCanSeek = blobContents.Stream.CanSeek;
+			
 			// no using statement as the blob contents will take ownership of this buffered payload and dispose it when the contents it is disposed
-			IBufferedPayload bufferedPayload = await _bufferedPayloadFactory.CreateFromStreamAsync(blobContents.Stream, blobContents.Length, "populate-store", cancellationToken);
+			IBufferedPayload? bufferedPayload = null;
+			if (!streamCanSeek)
+			{
+				bufferedPayload = await _bufferedPayloadFactory.CreateFromStreamAsync(blobContents.Stream, blobContents.Length, "populate-store", cancellationToken);
+			}
 
 			// Don't populate the last store, as that is where we got the hit
 			for (int i = 0; i < numStoreMisses; i++)
@@ -439,14 +445,44 @@ public class BlobService : IBlobService
 				using TelemetrySpan scope = _tracer.StartActiveSpan("HierarchicalStore.PopulateStore")
 					.SetAttribute("operation.name", "HierarchicalStore.PopulateStore")
 					.SetAttribute("BlobStore", blobStore.GetType().Name);
-				await using Stream s = bufferedPayload.GetStream();
-				// Populate each store traversed that did not have the content found lower in the hierarchy
-				await blobStore.PutObjectAsync(ns, s, blob);
+
+				Stream? s = null;
+				try
+				{
+					s = streamCanSeek ? blobContents.Stream : bufferedPayload!.GetStream();
+
+					// Populate each store traversed that did not have the content found lower in the hierarchy
+					await blobStore.PutObjectAsync(ns, s, blob);
+				}
+				finally
+				{
+					if (s != null)
+					{
+						if (streamCanSeek)
+						{
+							// it is a seekable stream so we just reset this
+							s.Seek(0, SeekOrigin.Begin);
+						}
+						else
+						{
+							await s.DisposeAsync();
+						}
+					}
+				}
+
+			}
+
+			if (streamCanSeek)
+			{
+				return blobContents;
+			}
+			else
+			{
+				// dispose the old blob contents before returning a new one
+				await blobContents.DisposeAsync();
+				return new BlobContents(bufferedPayload!);
 			}
 
-			// dispose the old blob contents before returning a new one
-			await blobContents.DisposeAsync();
-			return new BlobContents(bufferedPayload);
 		}
 
 		return blobContents;
