Cloud DDC - Added endpoint to suggest servers/caches to connect to based on public internet ips. Based on how we route local users to caches for other internal infrastructure.

Intended as a replacement for the current discovery document feature.

diff --git a/Engine/Source/Programs/UnrealCloudDDC/Jupiter/Controllers/StatusController.cs b/Engine/Source/Programs/UnrealCloudDDC/Jupiter/Controllers/StatusController.cs
--- a/Engine/Source/Programs/UnrealCloudDDC/Jupiter/Controllers/StatusController.cs
+++ b/Engine/Source/Programs/UnrealCloudDDC/Jupiter/Controllers/StatusController.cs
@@ -2,13 +2,16 @@
 
 using System;
 using System.Collections.Generic;
+using System.ComponentModel.DataAnnotations;
 using System.Linq;
+using System.Net;
 using System.Reflection;
 using System.Text.Json.Serialization;
 using EpicGames.Serialization;
 using Jupiter.Implementation;
 using Microsoft.AspNetCore.Authorization;
 using Microsoft.AspNetCore.Mvc;
+using Microsoft.Extensions.Logging;
 using Microsoft.Extensions.Options;
 
 namespace Jupiter.Controllers
@@ -22,13 +25,15 @@ namespace Jupiter.Controllers
 		private readonly IOptionsMonitor<JupiterSettings> _jupiterSettings;
 		private readonly IOptionsMonitor<ClusterSettings> _clusterSettings;
 		private readonly IPeerStatusService _statusService;
+		private readonly ILogger<StatusController> _logger;
 
-		public StatusController(VersionFile versionFile, IOptionsMonitor<JupiterSettings> jupiterSettings, IOptionsMonitor<ClusterSettings> clusterSettings, IPeerStatusService statusService)
+		public StatusController(VersionFile versionFile, IOptionsMonitor<JupiterSettings> jupiterSettings, IOptionsMonitor<ClusterSettings> clusterSettings, IPeerStatusService statusService, ILogger<StatusController> logger)
 		{
 			_versionFile = versionFile;
 			_jupiterSettings = jupiterSettings;
 			_clusterSettings = clusterSettings;
 			_statusService = statusService;
+			_logger = logger;
 		}
 
 		/// <summary>
@@ -82,6 +87,93 @@ namespace Jupiter.Controllers
 		{
 			return Ok(new PeersResponse(_jupiterSettings, _clusterSettings, includeInternalEndpoints, _statusService));
 		}
+
+		/// <summary>
+		/// Get a list of which servers you can talk to based on the ip used to call the endpoint
+		/// </summary>
+		/// <remarks>
+		/// </remarks>
+		/// <returns></returns>
+		[HttpGet("servers")]
+		[ProducesResponseType(type: typeof(PeersResponse), 200)]
+		public IActionResult Servers([FromQuery] IPAddress? asIp = null)
+		{
+			IPAddress? ip = asIp ?? Request.HttpContext.Connection.RemoteIpAddress;
+
+			if (_clusterSettings.CurrentValue.Discovery == null)
+			{
+				return BadRequest("No DiscoverySettings specified on the cluster settings");
+			}
+
+			(List<EndpointDefinition> cacheEndpoints, List<EndpointDefinition> serverEndpoints) = DetermineBestEndpoints(ip, _clusterSettings.CurrentValue.Discovery);
+
+			return Ok(new ServersResponse(cacheEndpoints, serverEndpoints));
+		}
+
+		private (List<EndpointDefinition>, List<EndpointDefinition>) DetermineBestEndpoints(IPAddress? ip, DiscoverySettings discoverySettings)
+		{
+			if (ip == null)
+			{
+				return (ResolveEndpoints(discoverySettings, discoverySettings.DefaultRule.CacheIds), ResolveEndpoints(discoverySettings, discoverySettings.DefaultRule.ServerIds));
+			}
+
+			List<string> cacheIds = new List<string>();
+			List<string> serverIds = new List<string>();
+
+			bool hadMatch = false;
+			foreach (RoutingRule rule in discoverySettings.Rules)
+			{
+				bool match = false;
+				foreach (string subnet in rule.Subnets)
+				{
+					// determine if subnet matches ip
+					if (IPNetwork.TryParse(subnet, out IPNetwork network))
+					{
+						if (network.Contains(ip))
+						{
+							match = true;
+							break;
+						}
+					}
+					else
+					{
+						_logger.LogWarning("Subnet {Subnet} in rule {Rule} does not to be in a CIDR format", subnet, rule.Name);
+					}
+				}
+
+				if (match)
+				{
+					cacheIds.AddRange(rule.CacheIds);
+					serverIds.AddRange(rule.ServerIds);
+					hadMatch = true;
+				}
+			}
+
+			if (!hadMatch)
+			{
+				return (ResolveEndpoints(discoverySettings, discoverySettings.DefaultRule.CacheIds), ResolveEndpoints(discoverySettings, discoverySettings.DefaultRule.ServerIds));
+			}
+
+			return (ResolveEndpoints(discoverySettings, cacheIds.Distinct()), ResolveEndpoints(discoverySettings, serverIds.Distinct()));
+		}
+
+		private List<EndpointDefinition> ResolveEndpoints(DiscoverySettings discoverySettings, IEnumerable<string> endpoints)
+		{
+			List<EndpointDefinition> resolvedEndpoints = new List<EndpointDefinition>();
+			foreach (string endpoint in endpoints)
+			{
+				if (discoverySettings.EndpointDefinitions.TryGetValue(endpoint, out EndpointDefinition? value))
+				{
+					resolvedEndpoints.Add(value);
+				}
+				else
+				{
+					_logger.LogWarning("Unable to resolve endpoint id {EndpointId} to a definition, make sure this exists in your discovery settings", endpoint);
+				}
+			}
+
+			return resolvedEndpoints;
+		}
 	}
 
 	public class PeersResponse
@@ -112,6 +204,62 @@ namespace Jupiter.Controllers
 		public List<KnownPeer> Peers { get; set; } = new List<KnownPeer>();
 	}
 
+	public class ServersResponse
+	{
+		public List<EndpointDefinition> CacheEndpoints { get; }
+		public List<EndpointDefinition> ServerEndpoints { get; }
+
+		public ServersResponse(List<EndpointDefinition> cacheEndpoints, List<EndpointDefinition> serverEndpoints)
+		{
+			CacheEndpoints = cacheEndpoints;
+			ServerEndpoints = serverEndpoints;
+		}
+	}
+
+	public class EndpointDefinition
+	{
+		[Required]
+		public string Name { get; set; } = null!;
+
+		[Required]
+		public Uri BaseUrl { get; set; } = null!;
+
+		public bool AssumeHttp2 { get; set; } = false;
+	}
+
+	public class RoutingRule
+	{
+		[Required] 
+		public string Name { get; set; } = null!;
+		
+		[Required] 
+		public string Description { get; set; } = null!;
+
+		[Required]
+		[System.Diagnostics.CodeAnalysis.SuppressMessage("Usage", "CA2227:Collection properties should be read only", Justification = "Used by serialization")]
+		public List<string> Subnets { get; set; } = new List<string>();
+
+		[System.Diagnostics.CodeAnalysis.SuppressMessage("Usage", "CA2227:Collection properties should be read only", Justification = "Used by serialization")]
+		public List<string> CacheIds { get; set; } = new List<string>();
+		
+		[System.Diagnostics.CodeAnalysis.SuppressMessage("Usage", "CA2227:Collection properties should be read only", Justification = "Used by serialization")]
+		public List<string> ServerIds { get; set; } = new List<string>();
+	}
+
+	public class DiscoverySettings
+	{
+		[Required]
+		[System.Diagnostics.CodeAnalysis.SuppressMessage("Usage", "CA2227:Collection properties should be read only", Justification = "Used by serialization")]
+		public Dictionary<string , EndpointDefinition> EndpointDefinitions { get; set; } = new Dictionary<string, EndpointDefinition>();
+
+		[Required]
+		public RoutingRule DefaultRule { get; set; } = null!;
+		
+		[Required]
+		[System.Diagnostics.CodeAnalysis.SuppressMessage("Usage", "CA2227:Collection properties should be read only", Justification = "Used by serialization")]
+		public List<RoutingRule> Rules { get; set; } = new List<RoutingRule>();
+	}
+
 	public class KnownPeer
 	{
 		public KnownPeer()
diff --git a/Engine/Source/Programs/UnrealCloudDDC/Jupiter/Settings.cs b/Engine/Source/Programs/UnrealCloudDDC/Jupiter/Settings.cs
--- a/Engine/Source/Programs/UnrealCloudDDC/Jupiter/Settings.cs
+++ b/Engine/Source/Programs/UnrealCloudDDC/Jupiter/Settings.cs
@@ -5,6 +5,7 @@ using System.Collections.Generic;
 using System.ComponentModel.DataAnnotations;
 using System.Linq;
 using Jupiter.Common;
+using Jupiter.Controllers;
 using Microsoft.Extensions.Caching.Memory;
 
 namespace Jupiter
@@ -262,6 +263,8 @@ namespace Jupiter
 	{
 		[System.Diagnostics.CodeAnalysis.SuppressMessage("Usage", "CA2227:Collection properties should be read only", Justification = "Used by serialization")]
 		public List<PeerSettings> Peers { get; set; } = new List<PeerSettings>();
+
+		public DiscoverySettings? Discovery { get; set; } = null;
 	}
 
 	public class PeerSettings
