Cloud DDC - Require a enabled option for our extra services, breaking change but gives us more control over when a load balancer or service is created in different regions. Allows you have to have a definition of a service that is only selectivley enabled.

diff --git a/Engine/Source/Programs/UnrealCloudDDC/Helm/UnrealCloudDDC/templates/extra-service-2.yaml b/Engine/Source/Programs/UnrealCloudDDC/Helm/UnrealCloudDDC/templates/extra-service-2.yaml
--- a/Engine/Source/Programs/UnrealCloudDDC/Helm/UnrealCloudDDC/templates/extra-service-2.yaml
+++ b/Engine/Source/Programs/UnrealCloudDDC/Helm/UnrealCloudDDC/templates/extra-service-2.yaml
@@ -1,4 +1,4 @@
-{{- if .Values.extraService2 }}
+{{- if .Values.extraService2.enabled }}
 apiVersion: v1
 kind: Service
 metadata:
diff --git a/Engine/Source/Programs/UnrealCloudDDC/Helm/UnrealCloudDDC/templates/extra-service-3.yaml b/Engine/Source/Programs/UnrealCloudDDC/Helm/UnrealCloudDDC/templates/extra-service-3.yaml
--- a/Engine/Source/Programs/UnrealCloudDDC/Helm/UnrealCloudDDC/templates/extra-service-3.yaml
+++ b/Engine/Source/Programs/UnrealCloudDDC/Helm/UnrealCloudDDC/templates/extra-service-3.yaml
@@ -1,4 +1,4 @@
-{{- if .Values.extraService3 }}
+{{- if .Values.extraService3.enabled }}
 apiVersion: v1
 kind: Service
 metadata:
diff --git a/Engine/Source/Programs/UnrealCloudDDC/Helm/UnrealCloudDDC/templates/extra-service.yaml b/Engine/Source/Programs/UnrealCloudDDC/Helm/UnrealCloudDDC/templates/extra-service.yaml
--- a/Engine/Source/Programs/UnrealCloudDDC/Helm/UnrealCloudDDC/templates/extra-service.yaml
+++ b/Engine/Source/Programs/UnrealCloudDDC/Helm/UnrealCloudDDC/templates/extra-service.yaml
@@ -1,4 +1,4 @@
-{{- if .Values.extraService }}
+{{- if .Values.extraService.enabled }}
 apiVersion: v1
 kind: Service
 metadata:
diff --git a/Engine/Source/Programs/UnrealCloudDDC/Helm/UnrealCloudDDC/values.yaml b/Engine/Source/Programs/UnrealCloudDDC/Helm/UnrealCloudDDC/values.yaml
--- a/Engine/Source/Programs/UnrealCloudDDC/Helm/UnrealCloudDDC/values.yaml
+++ b/Engine/Source/Programs/UnrealCloudDDC/Helm/UnrealCloudDDC/values.yaml
@@ -231,6 +231,15 @@ service:
   ##
   # loadBalancerIP:
 
+extraService:
+  enabled: false
+
+extraService2:
+  enabled: false
+
+extraService3:
+  enabled: false
+
 ## Configure the ingress resource that allows you to access the pod
 ##
 ingress:
@@ -394,6 +403,9 @@ dashboard:
   ##
   enabled: False
 
+  service:
+    enabled: False
+  
   ingress:
     ## Set to enable creating a ingress (load balancer)
     enabled: False
