Cloud DDC - Fixed null dereference while processing a exception when multipart downloading

diff --git a/Engine/Source/Programs/UnrealCloudDDC/Jupiter/Implementation/Blob/AmazonS3Store.cs b/Engine/Source/Programs/UnrealCloudDDC/Jupiter/Implementation/Blob/AmazonS3Store.cs
--- a/Engine/Source/Programs/UnrealCloudDDC/Jupiter/Implementation/Blob/AmazonS3Store.cs
+++ b/Engine/Source/Programs/UnrealCloudDDC/Jupiter/Implementation/Blob/AmazonS3Store.cs
@@ -543,7 +543,11 @@ namespace Jupiter.Implementation
 					{
 						for (int i = 0; i < partsCount; i++)
 						{
-							ArrayPool<byte>.Shared.Return(parts[i]);
+							// ReSharper disable once ConditionIsAlwaysTrueOrFalseAccordingToNullableAPIContract
+							if (parts[i] != null)
+							{
+								ArrayPool<byte>.Shared.Return(parts[i]);
+							}
 						}
 					}
 				}
