Cloud DDC - Fixed potential memory leak when using multipart download from S3, if a exception was thrown some of the shared buffers was not always returned.

diff --git a/Engine/Source/Programs/UnrealCloudDDC/Jupiter/Implementation/Blob/AmazonS3Store.cs b/Engine/Source/Programs/UnrealCloudDDC/Jupiter/Implementation/Blob/AmazonS3Store.cs
--- a/Engine/Source/Programs/UnrealCloudDDC/Jupiter/Implementation/Blob/AmazonS3Store.cs
+++ b/Engine/Source/Programs/UnrealCloudDDC/Jupiter/Implementation/Blob/AmazonS3Store.cs
@@ -487,11 +487,10 @@ namespace Jupiter.Implementation
 					try
 					{
 						{
-							byte[] partBuffer = ArrayPool<byte>.Shared.Rent((int)response.ContentLength);
-							await using MemoryStream ms = new MemoryStream(partBuffer);
-							await response.ResponseStream.CopyToAsync(ms, cancellationToken);
-							parts[0] = partBuffer;
+							parts[0] = ArrayPool<byte>.Shared.Rent((int)response.ContentLength);
 							lengths[0] = response.ContentLength;
+							await using MemoryStream ms = new MemoryStream(parts[0]);
+							await response.ResponseStream.CopyToAsync(ms, cancellationToken);
 						}
 
 						// multipart object, download it in parts, we already have the first part
@@ -500,13 +499,11 @@ namespace Jupiter.Implementation
 							GetObjectRequest partRequest = new GetObjectRequest { BucketName = _bucketName, Key = path, PartNumber = part };
 							GetObjectResponse partResponse = await _amazonS3.GetObjectAsync(partRequest, cancellationToken);
 
-							byte[] partBuffer = ArrayPool<byte>.Shared.Rent((int)partResponse.ContentLength);
-							await using MemoryStream ms = new MemoryStream(partBuffer);
-							await partResponse.ResponseStream.CopyToAsync(ms, cancellationToken);
-
 							// the first part was retrieved in the initial request
-							parts[part-1] = partBuffer;
+							parts[part-1] = ArrayPool<byte>.Shared.Rent((int)partResponse.ContentLength);
 							lengths[part-1] = partResponse.ContentLength;
+							await using MemoryStream ms = new MemoryStream(parts[part-1]);
+							await partResponse.ResponseStream.CopyToAsync(ms, cancellationToken);
 						});
 
 						long totalSize = lengths.Sum();
@@ -543,11 +540,7 @@ namespace Jupiter.Implementation
 					{
 						for (int i = 0; i < partsCount; i++)
 						{
-							// ReSharper disable once ConditionIsAlwaysTrueOrFalseAccordingToNullableAPIContract
-							if (parts[i] != null)
-							{
-								ArrayPool<byte>.Shared.Return(parts[i]);
-							}
+							ArrayPool<byte>.Shared.Return(parts[i]);
 						}
 					}
 				}
