Cloud DDC - Support compact binary serilization on the servers endpoint

diff --git a/Engine/Source/Programs/UnrealCloudDDC/Jupiter/Controllers/StatusController.cs b/Engine/Source/Programs/UnrealCloudDDC/Jupiter/Controllers/StatusController.cs
--- a/Engine/Source/Programs/UnrealCloudDDC/Jupiter/Controllers/StatusController.cs
+++ b/Engine/Source/Programs/UnrealCloudDDC/Jupiter/Controllers/StatusController.cs
@@ -206,7 +206,9 @@ namespace Jupiter.Controllers
 
 	public class ServersResponse
 	{
+		[CbField("cacheEndpoints")]
 		public List<EndpointDefinition> CacheEndpoints { get; }
+		[CbField("serverEndpoints")]
 		public List<EndpointDefinition> ServerEndpoints { get; }
 
 		public ServersResponse(List<EndpointDefinition> cacheEndpoints, List<EndpointDefinition> serverEndpoints)
@@ -219,11 +221,14 @@ namespace Jupiter.Controllers
 	public class EndpointDefinition
 	{
 		[Required]
+		[CbField("name")]
 		public string Name { get; set; } = null!;
 
 		[Required]
+		[CbField("baseUrl")]
 		public Uri BaseUrl { get; set; } = null!;
 
+		[CbField("assumeHttp2")]
 		public bool AssumeHttp2 { get; set; } = false;
 	}
 
