Cloud DDC - Added ability to specify oidc configuration for clients which is served up on an unauthenticated endpoint. Basic encryption is applied to this data but as the key is available in source it should not be considered secure.
The oidc configuration is not secret data though so this is okay.

diff --git a/Engine/Source/Programs/UnrealCloudDDC/Jupiter.Common/BaseStartup.cs b/Engine/Source/Programs/UnrealCloudDDC/Jupiter.Common/BaseStartup.cs
--- a/Engine/Source/Programs/UnrealCloudDDC/Jupiter.Common/BaseStartup.cs
+++ b/Engine/Source/Programs/UnrealCloudDDC/Jupiter.Common/BaseStartup.cs
@@ -772,6 +772,17 @@ namespace Jupiter
 		/// </summary>
 		public bool UseLegacyConfiguration { get; set; } = false;
 
+		/// <summary>
+		/// Configuration read by clients on how to interactively login locally
+		/// </summary>
+		[System.Diagnostics.CodeAnalysis.SuppressMessage("Usage", "CA2227:Collection properties should be read only", Justification = "Used by the configuration system")]
+		public Dictionary<string, string> ClientOidcConfiguration { get; set; } = new Dictionary<string, string>();
+
+		/// <summary>
+		/// The encryption key used to encrypt the configuration - usually not modified as you then need to also distribute this key. This needs to be exactly 16 bytes
+		/// </summary>
+		public string ClientOidcEncryptionKey { get; set; } = "892a27ef5cbf4894af2e6bd53a54aa48";
+
 		public IEnumerable<ValidationResult> Validate(ValidationContext validationContext)
 		{
 			List<ValidationResult> validationResults = new List<ValidationResult>();
diff --git a/Engine/Source/Programs/UnrealCloudDDC/Jupiter/Controllers/AuthController.cs b/Engine/Source/Programs/UnrealCloudDDC/Jupiter/Controllers/AuthController.cs
--- a/Engine/Source/Programs/UnrealCloudDDC/Jupiter/Controllers/AuthController.cs
+++ b/Engine/Source/Programs/UnrealCloudDDC/Jupiter/Controllers/AuthController.cs
@@ -2,6 +2,10 @@
 
 using System.Collections.Generic;
 using System.ComponentModel.DataAnnotations;
+using System.IO;
+using System.Net.Mime;
+using System.Security.Cryptography;
+using System.Text.Json;
 using System.Threading.Tasks;
 using EpicGames.Core;
 using EpicGames.Horde.Storage;
@@ -26,6 +30,34 @@ namespace Jupiter.Controllers
 			_authSettings = authSettings;
 		}
 
+		[HttpGet("config")]
+		// disable authentication on this endpoint
+		[AllowAnonymous]
+		// this endpoint always produces encrypted json
+		[Produces(MediaTypeNames.Application.Octet)]
+		public async Task<ActionResult> GetConfigAsync()
+		{
+			byte[] b = JsonSerializer.SerializeToUtf8Bytes(_authSettings.CurrentValue.ClientOidcConfiguration);
+
+			using Aes aes = Aes.Create();
+			byte[] key = StringUtils.ToHashFromHexString(_authSettings.CurrentValue.ClientOidcEncryptionKey);
+			aes.Key = key;
+			aes.GenerateIV();
+			// write the IV into the stream before the encrypted content
+			Stream responseStream = new MemoryStream();
+			await responseStream.WriteAsync(aes.IV);
+			await using CryptoStream cryptoStream = new(responseStream, aes.CreateEncryptor(), CryptoStreamMode.Write);
+			await cryptoStream.WriteAsync(b, 0, b.Length);
+			await cryptoStream.FlushFinalBlockAsync();
+
+			responseStream.Position = 0;
+			Response.ContentType = MediaTypeNames.Application.Octet;
+			Response.StatusCode = 200;
+			Response.ContentLength = responseStream.Length;
+			await responseStream.CopyToAsync(Response.Body);
+			return new EmptyResult();
+		}
+
 		[HttpGet("{ns}")]
 		public async Task<IActionResult> VerifyAsync([FromRoute][Required] NamespaceId ns)
 		{
