Cloud DDC - Added default constructor to servers response. Fixes issue when attempting to convert to compact binary.

diff --git a/Engine/Source/Programs/UnrealCloudDDC/Jupiter/Controllers/StatusController.cs b/Engine/Source/Programs/UnrealCloudDDC/Jupiter/Controllers/StatusController.cs
--- a/Engine/Source/Programs/UnrealCloudDDC/Jupiter/Controllers/StatusController.cs
+++ b/Engine/Source/Programs/UnrealCloudDDC/Jupiter/Controllers/StatusController.cs
@@ -206,10 +206,18 @@ namespace Jupiter.Controllers
 
 	public class ServersResponse
 	{
+		public ServersResponse()
+		{
+
+		}
+
 		[CbField("cacheEndpoints")]
-		public List<EndpointDefinition> CacheEndpoints { get; }
+		[System.Diagnostics.CodeAnalysis.SuppressMessage("Usage", "CA2227:Collection properties should be read only", Justification = "Used by serialization")]
+		public List<EndpointDefinition> CacheEndpoints { get; set; } = new List<EndpointDefinition>();
+		
 		[CbField("serverEndpoints")]
-		public List<EndpointDefinition> ServerEndpoints { get; }
+		[System.Diagnostics.CodeAnalysis.SuppressMessage("Usage", "CA2227:Collection properties should be read only", Justification = "Used by serialization")]
+		public List<EndpointDefinition> ServerEndpoints { get; set; } = new List<EndpointDefinition>();
 
 		public ServersResponse(List<EndpointDefinition> cacheEndpoints, List<EndpointDefinition> serverEndpoints)
 		{
