Cloud DDC - Fixes to reporting of old namespaces/buckets

diff --git a/Engine/Source/Programs/UnrealCloudDDC/Jupiter/Implementation/GC/CacheCleanup.cs b/Engine/Source/Programs/UnrealCloudDDC/Jupiter/Implementation/GC/CacheCleanup.cs
--- a/Engine/Source/Programs/UnrealCloudDDC/Jupiter/Implementation/GC/CacheCleanup.cs
+++ b/Engine/Source/Programs/UnrealCloudDDC/Jupiter/Implementation/GC/CacheCleanup.cs
@@ -163,27 +163,32 @@ namespace Jupiter.Implementation
 
 			await foreach (NamespaceId ns in _referencesStore.GetNamespacesAsync(cancellationToken))
 			{
-				if (namespacesWithRefs.ContainsKey(ns))
+				bool refFound = false;
+				if (!namespacesWithRefs.TryGetValue(ns, out ConcurrentDictionary<BucketId, bool>? foundBuckets))
 				{
-					continue;
+					foundBuckets = new ConcurrentDictionary<BucketId, bool>();
 				}
 
-				ConcurrentDictionary<BucketId, bool> foundBuckets = namespacesWithRefs[ns];
 				await foreach (BucketId bucket in _referencesStore.GetBucketsAsync(ns, cancellationToken))
 				{
 					// if the bucket had a ref we should not remove its tracking
 					if (foundBuckets.ContainsKey(bucket))
 					{
+						refFound = true;
 						continue;
 					}
 					_logger.LogWarning("Namespace {Namespace} Bucket {Bucket} has no refs. Should be removed from tracking.", ns, bucket);
 					// TODO: We need a ligher way to drop a bucket that has no refs
 					// await _referencesStore.DeleteBucketAsync(ns, bucket, cancellationToken);
 				}
-				_logger.LogWarning("Namespace {Namespace} has no refs. Should be removed from tracking.", ns);
 
-				// TODO: We should have a lighter namespace drop that does not remove refs, option to it?
-				//await _referencesStore.DropNamespaceAsync(ns, cancellationToken);
+				if (!refFound)
+				{
+					_logger.LogWarning("Namespace {Namespace} has no refs. Should be removed from tracking.", ns);
+
+					// TODO: We should have a lighter namespace drop that does not remove refs, option to it?
+					//await _referencesStore.DropNamespaceAsync(ns, cancellationToken);
+				}
 			}
 			return countOfDeletedRecords;
 		}
