Cloud DDC - Removed use of unacceptable word

diff --git a/Engine/Source/Programs/UnrealCloudDDC/Jupiter.Common/BaseStartup.cs b/Engine/Source/Programs/UnrealCloudDDC/Jupiter.Common/BaseStartup.cs
--- a/Engine/Source/Programs/UnrealCloudDDC/Jupiter.Common/BaseStartup.cs
+++ b/Engine/Source/Programs/UnrealCloudDDC/Jupiter.Common/BaseStartup.cs
@@ -761,7 +761,7 @@ namespace Jupiter
 		public Uri? RedirectUri { get; set; } = null;
 		
 		/// <summary>
-		/// List of redirect uris that can be used (are whitelisted in the IdP). In priority order.
+		/// List of redirect uris that can be used (are allow listed in the IdP). In priority order.
 		/// </summary>
 		[System.Diagnostics.CodeAnalysis.SuppressMessage("Usage", "CA2227:Collection properties should be read only", Justification = "Used by the configuration system")]
 		public List<Uri>? PossibleRedirectUri { get; set; } = null!;
