Cloud DDC - Missing file when adding compressed hash response

diff --git a/Engine/Source/Programs/UnrealCloudDDC/Jupiter/Implementation/Blob/IBlobService.cs b/Engine/Source/Programs/UnrealCloudDDC/Jupiter/Implementation/Blob/IBlobService.cs
--- a/Engine/Source/Programs/UnrealCloudDDC/Jupiter/Implementation/Blob/IBlobService.cs
+++ b/Engine/Source/Programs/UnrealCloudDDC/Jupiter/Implementation/Blob/IBlobService.cs
@@ -147,7 +147,7 @@ public static class BlobServiceExtensions
 		return cid;
 	}
 
-	public static async Task<(BlobContents, string)> GetCompressedObjectAsync(this IBlobService blobService, NamespaceId ns, ContentId contentId, IServiceProvider provider, bool supportsRedirectUri = false, CancellationToken cancellationToken = default)
+	public static async Task<(BlobContents, string, BlobId?)> GetCompressedObjectAsync(this IBlobService blobService, NamespaceId ns, ContentId contentId, IServiceProvider provider, bool supportsRedirectUri = false, CancellationToken cancellationToken = default)
 	{
 		IContentIdStore contentIdStore = provider.GetService<IContentIdStore>()!;
 		Tracer tracer = provider.GetService<Tracer>()!;
@@ -169,7 +169,7 @@ public static class BlobServiceExtensions
 				mimeType = MediaTypeNames.Application.Octet;
 			}
 
-			return (await blobService.GetObjectAsync(ns, blobToReturn, supportsRedirectUri: supportsRedirectUri, cancellationToken: cancellationToken), mimeType);
+			return (await blobService.GetObjectAsync(ns, blobToReturn, supportsRedirectUri: supportsRedirectUri, cancellationToken: cancellationToken), mimeType, blobToReturn);
 		}
 
 		// chunked content, combine the chunks into a single stream
@@ -191,8 +191,9 @@ public static class BlobServiceExtensions
 
 		ms.Seek(0, SeekOrigin.Begin);
 
+		// TODO: A chunked blob does not store the hash of the full object so we are unable to return it, only used for a header so okay for now but should be fixed
 		// chunking could not have happened for a non compressed buffer so assume it is compressed
-		return (new BlobContents(ms, ms.Length), CustomMediaTypeNames.UnrealCompressedBuffer);
+		return (new BlobContents(ms, ms.Length), CustomMediaTypeNames.UnrealCompressedBuffer, null);
 	}
 }
 
