Cloud DDC - Reverted send of header with compressed hash as this is causing warnings form zen

diff --git a/Engine/Source/Programs/UnrealCloudDDC/Jupiter/Controllers/BuildsController.cs b/Engine/Source/Programs/UnrealCloudDDC/Jupiter/Controllers/BuildsController.cs
--- a/Engine/Source/Programs/UnrealCloudDDC/Jupiter/Controllers/BuildsController.cs
+++ b/Engine/Source/Programs/UnrealCloudDDC/Jupiter/Controllers/BuildsController.cs
@@ -916,10 +916,10 @@ namespace Jupiter.Controllers
 					return new UnsupportedMediaTypeResult();
 				}
 
-				if (contentHash != null)
+				/*if (contentHash != null)
 				{
 					Response.Headers[CommonHeaders.HashHeaderName] = contentHash.ToString();
-				}
+				}*/
 
 				if (blobContents.RedirectUri != null)
 				{
diff --git a/Engine/Source/Programs/UnrealCloudDDC/Jupiter/Controllers/CompressedBlobController.cs b/Engine/Source/Programs/UnrealCloudDDC/Jupiter/Controllers/CompressedBlobController.cs
--- a/Engine/Source/Programs/UnrealCloudDDC/Jupiter/Controllers/CompressedBlobController.cs
+++ b/Engine/Source/Programs/UnrealCloudDDC/Jupiter/Controllers/CompressedBlobController.cs
@@ -78,10 +78,10 @@ namespace Jupiter.Controllers
 					return new UnsupportedMediaTypeResult();
 				}
 
-				if (contentHash != null)
+				/*if (contentHash != null)
 				{
 					Response.Headers[CommonHeaders.HashHeaderName] = contentHash.ToString();
-				}
+				}*/
 
 				if (blobContents.RedirectUri != null)
 				{
diff --git a/Engine/Source/Programs/UnrealCloudDDC/Jupiter/Controllers/ReferencesController.cs b/Engine/Source/Programs/UnrealCloudDDC/Jupiter/Controllers/ReferencesController.cs
--- a/Engine/Source/Programs/UnrealCloudDDC/Jupiter/Controllers/ReferencesController.cs
+++ b/Engine/Source/Programs/UnrealCloudDDC/Jupiter/Controllers/ReferencesController.cs
@@ -1071,10 +1071,10 @@ namespace Jupiter.Controllers
 					return new UnsupportedMediaTypeResult();
 				}
 
-				if (contentHash != null)
+				/*if (contentHash != null)
 				{
 					Response.Headers[CommonHeaders.HashHeaderName] = contentHash.ToString();
-				}
+				}*/
 
 				if (blobContents.RedirectUri != null)
 				{
