Cloud DDC - Simplified response from actions endpoints when testing which actions a particular user has access to. Also added which policies are matched to the response for better understanding of what granted the action.

diff --git a/Engine/Source/Programs/UnrealCloudDDC/Jupiter/Controllers/AuthController.cs b/Engine/Source/Programs/UnrealCloudDDC/Jupiter/Controllers/AuthController.cs
--- a/Engine/Source/Programs/UnrealCloudDDC/Jupiter/Controllers/AuthController.cs
+++ b/Engine/Source/Programs/UnrealCloudDDC/Jupiter/Controllers/AuthController.cs
@@ -4,12 +4,14 @@ using System;
 using System.Collections.Generic;
 using System.ComponentModel.DataAnnotations;
 using System.IO;
+using System.Linq;
 using System.Net.Mime;
 using System.Security.Cryptography;
 using System.Text.Json;
 using System.Threading.Tasks;
 using EpicGames.Core;
 using EpicGames.Horde.Storage;
+using EpicGames.Serialization;
 using Microsoft.AspNetCore.Authorization;
 using Microsoft.AspNetCore.Mvc;
 using Microsoft.Extensions.Logging;
@@ -83,19 +85,28 @@ namespace Jupiter.Controllers
 			
 			CaptureLogger captureLogger = new CaptureLogger();
 			bool policiesFound = false;
+			Dictionary<string, List<JupiterAclAction>> matchedPolicies = new();
+			int suffixCounter = 0;
 			foreach (AclPolicy acl in _authSettings.CurrentValue.Policies)
 			{
 				policiesFound = true;
-				allowedActions.AddRange(acl.Resolve(User, new AccessScope(ns), captureLogger));
+				List<JupiterAclAction> grantedActions = acl.Resolve(User, new AccessScope(ns), captureLogger).ToList();
+				allowedActions.AddRange(grantedActions);
+				if (!matchedPolicies.TryAdd(acl.Name, grantedActions))
+				{
+					// the policy already existed likely because the name is not globally unique so we just append a random suffix
+					suffixCounter++;
+					matchedPolicies.TryAdd($"{acl.Name}_{suffixCounter}", grantedActions);
+				}
 			}
 
 			if (!policiesFound)
 			{
 				captureLogger.LogWarning("No policies set so no actions generated");
-				return BadRequest(new JsonResult(new { Actions = allowedActions, LogOutput = captureLogger.RenderLines() }));
+				return BadRequest(new ActionsResult {Actions = allowedActions, LogOutput = captureLogger.RenderLines(), MatchedPolicies = matchedPolicies });
 			}
 
-			return Ok(new JsonResult(new { Actions = allowedActions , LogOutput = captureLogger.RenderLines()}));
+			return Ok(new ActionsResult { Actions = allowedActions , LogOutput = captureLogger.RenderLines(), MatchedPolicies = matchedPolicies});
 		}
 		
 		[HttpGet("{ns}/{bucket}/actions")]
@@ -103,27 +114,46 @@ namespace Jupiter.Controllers
 		{
 			List<JupiterAclAction> allowedActions = new List<JupiterAclAction>();
 			
+			Dictionary<string, List<JupiterAclAction>> matchedPolicies = new();
+			int suffixCounter = 0;
 			CaptureLogger captureLogger = new CaptureLogger();
 			bool policiesFound = false;
 			foreach (AclPolicy acl in _authSettings.CurrentValue.Policies)
 			{
 				policiesFound = true;
-				allowedActions.AddRange(acl.Resolve(User, new AccessScope(ns, bucket), captureLogger));
+				List<JupiterAclAction> grantedActions = acl.Resolve(User, new AccessScope(ns, bucket), captureLogger).ToList();
+				allowedActions.AddRange(grantedActions);
+
+				if (!matchedPolicies.TryAdd(acl.Name, grantedActions))
+				{
+					// the policy already existed likely because the name is not globally unique so we just append a random suffix
+					suffixCounter++;
+					matchedPolicies.TryAdd($"{acl.Name}_{suffixCounter}", grantedActions);
+				}
 			}
 
 			if (!policiesFound)
 			{
 				captureLogger.LogWarning("No policies set so no actions generated");
-				return BadRequest(new JsonResult(new { Actions = allowedActions, LogOutput = captureLogger.RenderLines() }));
+				return BadRequest(new ActionsResult { Actions = allowedActions, LogOutput = captureLogger.RenderLines(), MatchedPolicies = matchedPolicies });
 			}
 
-			return Ok(new JsonResult(new { Actions = allowedActions , LogOutput = captureLogger.RenderLines()}));
+			return Ok(new ActionsResult { Actions = allowedActions , LogOutput = captureLogger.RenderLines(), MatchedPolicies = matchedPolicies});
 		}
 	}
 
 	public class ActionsResult
 	{
 		[System.Diagnostics.CodeAnalysis.SuppressMessage("Usage", "CA2227:Collection properties should be read only", Justification = "Used by serialization")]
+		[CbField("actions")]
 		public List<JupiterAclAction> Actions { get; set; } = new List<JupiterAclAction>();
+
+		[System.Diagnostics.CodeAnalysis.SuppressMessage("Usage", "CA2227:Collection properties should be read only", Justification = "Used by serialization")]
+		[CbField("logOutput")]
+		public List<string> LogOutput { get; set; } = new List<string>();
+
+		[System.Diagnostics.CodeAnalysis.SuppressMessage("Usage", "CA2227:Collection properties should be read only", Justification = "Used by serialization")]
+		[CbField("matchedPolicies")]
+		public Dictionary<string, List<JupiterAclAction>> MatchedPolicies { get; set; } = new Dictionary<string, List<JupiterAclAction>>();
 	}
 }
