Cloud DDC - New attempt at determining buckets and namespaces that are not in use and cleaning them up, this time based on the actual objects table rather then the last access tracking which isn't used for all cases.

diff --git a/Engine/Source/Programs/UnrealCloudDDC/Jupiter/Implementation/GC/CacheCleanup.cs b/Engine/Source/Programs/UnrealCloudDDC/Jupiter/Implementation/GC/CacheCleanup.cs
--- a/Engine/Source/Programs/UnrealCloudDDC/Jupiter/Implementation/GC/CacheCleanup.cs
+++ b/Engine/Source/Programs/UnrealCloudDDC/Jupiter/Implementation/GC/CacheCleanup.cs
@@ -1,7 +1,6 @@
 ﻿// Copyright Epic Games, Inc. All Rights Reserved.
 
 using System;
-using System.Collections.Concurrent;
 using System.Collections.Generic;
 using System.Diagnostics.Metrics;
 using System.Threading;
@@ -96,7 +95,7 @@ namespace Jupiter.Implementation
 			long consideredCount = 0;
 			DateTime cleanupStart = DateTime.Now;
 
-			ConcurrentDictionary<NamespaceId, ConcurrentDictionary<BucketId, bool>> namespacesWithRefs = new();
+			Task cleanupTrackedStateTask = _referencesStore.CleanupTrackedStateAsync();
 
 			await Parallel.ForEachAsync(_referencesStore.GetRecordsAsync(cancellationToken),
 				new ParallelOptions
@@ -107,10 +106,6 @@ namespace Jupiter.Implementation
 				{
 					(NamespaceId ns, BucketId bucket, RefId name, DateTime lastAccessTime) = tuple;
 
-					// keep track of which namespaces and buckets we have seen
-					ConcurrentDictionary<BucketId, bool> bucketTracking = namespacesWithRefs.GetOrAdd(ns, new ConcurrentDictionary<BucketId, bool>());
-					bucketTracking.TryAdd(bucket, true);
-
 					if (!ShouldGCNamespace(ns))
 					{
 						return;
@@ -157,39 +152,12 @@ namespace Jupiter.Implementation
 					}
 				});
 
+			await cleanupTrackedStateTask;
+
 			TimeSpan cleanupDuration = DateTime.Now - cleanupStart;
 			_logger.LogInformation(
 				"Finished cleaning refs. Refs considered: {ConsideredCount} Refs Deleted: {DeletedCount}. Cleanup took: {CleanupDuration}", consideredCount, countOfDeletedRecords, cleanupDuration);
 
-			await foreach (NamespaceId ns in _referencesStore.GetNamespacesAsync(cancellationToken))
-			{
-				bool refFound = false;
-				if (!namespacesWithRefs.TryGetValue(ns, out ConcurrentDictionary<BucketId, bool>? foundBuckets))
-				{
-					foundBuckets = new ConcurrentDictionary<BucketId, bool>();
-				}
-
-				await foreach (BucketId bucket in _referencesStore.GetBucketsAsync(ns, cancellationToken))
-				{
-					// if the bucket had a ref we should not remove its tracking
-					if (foundBuckets.ContainsKey(bucket))
-					{
-						refFound = true;
-						continue;
-					}
-					_logger.LogWarning("Namespace {Namespace} Bucket {Bucket} has no refs. Should be removed from tracking.", ns, bucket);
-					// TODO: We need a ligher way to drop a bucket that has no refs
-					// await _referencesStore.DeleteBucketAsync(ns, bucket, cancellationToken);
-				}
-
-				if (!refFound)
-				{
-					_logger.LogWarning("Namespace {Namespace} has no refs. Should be removed from tracking.", ns);
-
-					// TODO: We should have a lighter namespace drop that does not remove refs, option to it?
-					//await _referencesStore.DropNamespaceAsync(ns, cancellationToken);
-				}
-			}
 			return countOfDeletedRecords;
 		}
 
diff --git a/Engine/Source/Programs/UnrealCloudDDC/Jupiter/Implementation/Objects/CacheReferencesStore.cs b/Engine/Source/Programs/UnrealCloudDDC/Jupiter/Implementation/Objects/CacheReferencesStore.cs
--- a/Engine/Source/Programs/UnrealCloudDDC/Jupiter/Implementation/Objects/CacheReferencesStore.cs
+++ b/Engine/Source/Programs/UnrealCloudDDC/Jupiter/Implementation/Objects/CacheReferencesStore.cs
@@ -94,6 +94,12 @@ namespace Jupiter.Implementation
 			await Task.WhenAll(cacheUpdateLastAccess, upstreamUpdateLastAccess);
 		}
 
+		public Task CleanupTrackedStateAsync()
+		{
+			// nothing to cleanup
+			return Task.CompletedTask;
+		}
+
 		public IAsyncEnumerable<(NamespaceId, BucketId, RefId, DateTime)> GetRecordsAsync(CancellationToken cancellationToken)
 		{
 			throw new NotImplementedException("GetRecords not supported on a cached reference store");
@@ -241,6 +247,12 @@ namespace Jupiter.Implementation
 			await GetAsync(ns, bucket, key, IReferencesStore.FieldFlags.None, IReferencesStore.OperationFlags.None, cancellationToken);
 		}
 
+		public Task CleanupTrackedStateAsync()
+		{
+			// nothing to clean up
+			return Task.CompletedTask;
+		}
+
 		public IAsyncEnumerable<(NamespaceId, BucketId, RefId, DateTime)> GetRecordsAsync(CancellationToken cancellationToken)
 		{
 			throw new NotImplementedException("GetRecords is not supported on a upstream reference store");
diff --git a/Engine/Source/Programs/UnrealCloudDDC/Jupiter/Implementation/Objects/IReferencesStore.cs b/Engine/Source/Programs/UnrealCloudDDC/Jupiter/Implementation/Objects/IReferencesStore.cs
--- a/Engine/Source/Programs/UnrealCloudDDC/Jupiter/Implementation/Objects/IReferencesStore.cs
+++ b/Engine/Source/Programs/UnrealCloudDDC/Jupiter/Implementation/Objects/IReferencesStore.cs
@@ -44,6 +44,7 @@ namespace Jupiter.Implementation
 		Task<long> DropNamespaceAsync(NamespaceId ns, CancellationToken cancellationToken = default);
 		Task<long> DeleteBucketAsync(NamespaceId ns, BucketId bucket, CancellationToken cancellationToken = default);
 		Task UpdateTTL(NamespaceId ns, BucketId bucket, RefId refId, uint ttl, CancellationToken cancellationToken = default);
+		Task CleanupTrackedStateAsync();
 	}
 
 	public class RefRecord
diff --git a/Engine/Source/Programs/UnrealCloudDDC/Jupiter/Implementation/Objects/MemoryCachedReferencesStore.cs b/Engine/Source/Programs/UnrealCloudDDC/Jupiter/Implementation/Objects/MemoryCachedReferencesStore.cs
--- a/Engine/Source/Programs/UnrealCloudDDC/Jupiter/Implementation/Objects/MemoryCachedReferencesStore.cs
+++ b/Engine/Source/Programs/UnrealCloudDDC/Jupiter/Implementation/Objects/MemoryCachedReferencesStore.cs
@@ -209,6 +209,11 @@ namespace Jupiter.Implementation.Objects
 			_referenceCaches.Clear();
 		}
 
+		public Task CleanupTrackedStateAsync()
+		{
+			return _actualStore.CleanupTrackedStateAsync();
+		}
+
 		public IReferencesStore GetUnderlyingStore()
 		{
 			return _actualStore;
diff --git a/Engine/Source/Programs/UnrealCloudDDC/Jupiter/Implementation/Objects/MemoryReferencesStore.cs b/Engine/Source/Programs/UnrealCloudDDC/Jupiter/Implementation/Objects/MemoryReferencesStore.cs
--- a/Engine/Source/Programs/UnrealCloudDDC/Jupiter/Implementation/Objects/MemoryReferencesStore.cs
+++ b/Engine/Source/Programs/UnrealCloudDDC/Jupiter/Implementation/Objects/MemoryReferencesStore.cs
@@ -215,6 +215,12 @@ namespace Jupiter.Implementation
 			return Task.CompletedTask;
 		}
 
+		public Task CleanupTrackedStateAsync()
+		{
+			// nothing to cleanup
+			return Task.CompletedTask;
+		}
+
 		private static string BuildKey(NamespaceId ns, BucketId bucket, RefId name)
 		{
 			return $"{ns}.{bucket}.{name}";
diff --git a/Engine/Source/Programs/UnrealCloudDDC/Jupiter/Implementation/Objects/MongoReferencesStore.cs b/Engine/Source/Programs/UnrealCloudDDC/Jupiter/Implementation/Objects/MongoReferencesStore.cs
--- a/Engine/Source/Programs/UnrealCloudDDC/Jupiter/Implementation/Objects/MongoReferencesStore.cs
+++ b/Engine/Source/Programs/UnrealCloudDDC/Jupiter/Implementation/Objects/MongoReferencesStore.cs
@@ -287,6 +287,12 @@ namespace Jupiter.Implementation
 
 			AddIndexFor<MongoReferencesModelV0>().CreateOne(indexTTL);
 		}
+
+		public Task CleanupTrackedStateAsync()
+		{
+			// nothing to clean up
+			return Task.CompletedTask;
+		}
 	}
 
 	// we do versioning by writing a discriminator into object
diff --git a/Engine/Source/Programs/UnrealCloudDDC/Jupiter/Implementation/Objects/ScyllaReferencesStore.cs b/Engine/Source/Programs/UnrealCloudDDC/Jupiter/Implementation/Objects/ScyllaReferencesStore.cs
--- a/Engine/Source/Programs/UnrealCloudDDC/Jupiter/Implementation/Objects/ScyllaReferencesStore.cs
+++ b/Engine/Source/Programs/UnrealCloudDDC/Jupiter/Implementation/Objects/ScyllaReferencesStore.cs
@@ -34,6 +34,7 @@ namespace Jupiter.Implementation
 		private readonly PreparedStatement _getNamespacesStatement;
 		private readonly PreparedStatement _getNamespacesOldStatement;
 		private readonly PreparedStatement _getObjectsForPartitionRangeStatement;
+		private readonly PreparedStatement _getObjectsBucketsStatement;
 		private readonly PreparedStatement _getObjectsLastAccessForPartitionRangeStatement;
 
 		private readonly ConcurrentDictionary<NamespaceId, ConcurrentBag<BucketId>> _addedBuckets = new ConcurrentDictionary<NamespaceId, ConcurrentBag<BucketId>>();
@@ -94,6 +95,7 @@ namespace Jupiter.Implementation
 			_getNamespacesStatement = _session.Prepare("SELECT DISTINCT namespace FROM buckets_v2");
 			_getNamespacesOldStatement = _session.Prepare("SELECT DISTINCT namespace FROM buckets");
 
+			_getObjectsBucketsStatement = _session.Prepare("select namespace, bucket FROM objects");
 			_getObjectsForPartitionRangeStatement = _session.Prepare($"SELECT namespace, bucket, name, last_access_time FROM objects WHERE token(namespace, bucket, name) >= ? AND token(namespace, bucket, name) <= ? {cqlOptions}");
 			_getObjectsLastAccessForPartitionRangeStatement = _session.Prepare($"SELECT namespace, bucket, name, last_access_time FROM object_last_access_v2 WHERE token(namespace, bucket, name) >= ? AND token(namespace, bucket, name) <= ? {cqlOptions}");
 		}
@@ -605,6 +607,68 @@ namespace Jupiter.Implementation
 			}
 		}
 
+		public async Task CleanupTrackedStateAsync()
+		{
+			_logger.LogInformation("Starting internal scylla cleanup. ");
+			DateTime cleanupStart = DateTime.Now;
+			ConcurrentDictionary<NamespaceId, ConcurrentDictionary<BucketId, bool>> namespacesWithRefs = new();
+
+			// build lookup of which namespaces and buckets are actually references by live objects
+			{
+				RowSet rowSet = await _session.ExecuteAsync(_getObjectsBucketsStatement.Bind());
+
+				foreach (Row row in rowSet)
+				{
+					if (rowSet.GetAvailableWithoutFetching() == 0)
+					{
+						await rowSet.FetchMoreResultsAsync();
+					}
+
+					NamespaceId ns = new NamespaceId(row.GetValue<string>(0));
+					BucketId bucket = new BucketId(row.GetValue<string>(1));
+
+					ConcurrentDictionary<BucketId, bool> bucketTracking = namespacesWithRefs.GetOrAdd(ns, new ConcurrentDictionary<BucketId, bool>());
+					bucketTracking.TryAdd(bucket, true);
+				}
+			}
+
+			// enumerate our cached namespaces and buckets and determine if any of them are old
+			await foreach (NamespaceId ns in GetNamespacesAsync(CancellationToken.None))
+			{
+				bool refFound = false;
+				if (!namespacesWithRefs.TryGetValue(ns, out ConcurrentDictionary<BucketId, bool>? foundBuckets))
+				{
+					// namespace does not exist, still need to processes all buckets to remove their tracking
+					foundBuckets = new ConcurrentDictionary<BucketId, bool>();
+				}
+
+				await foreach (BucketId bucket in GetBucketsAsync(ns, CancellationToken.None))
+				{
+					// if the bucket had a ref we should not remove its tracking
+					if (foundBuckets.ContainsKey(bucket))
+					{
+						refFound = true;
+						continue;
+					}
+					_logger.LogWarning("Namespace {Namespace} Bucket {Bucket} has no refs. Should be removed from tracking.", ns, bucket);
+
+					// TODO: introduce option to actually delete this and not just report
+					//await _session.ExecuteAsync(new SimpleStatement("DELETE FROM buckets_v2 WHERE namespace = ? and bucket = ?", ns.ToString(), bucket.ToString()));
+				}
+
+				if (!refFound)
+				{
+					_logger.LogWarning("Namespace {Namespace} has no refs. Should be removed from tracking.", ns);
+
+					// TODO: introduce option to actually delete this and not just report
+					// await _session.ExecuteAsync(new SimpleStatement("DELETE FROM buckets_v2 WHERE namespace = ?", ns.ToString()));
+				}
+			}
+
+			TimeSpan cleanupDuration = DateTime.Now - cleanupStart;
+			_logger.LogInformation("Finished internal scylla cleanup. Cleanup took: {CleanupDuration}", cleanupDuration);
+		}
+
 		public async IAsyncEnumerable<RefId> GetRecordsInBucketAsync(NamespaceId ns, BucketId bucket, [EnumeratorCancellation] CancellationToken cancellationToken)
 		{
 			string nsAsString = ns.ToString();
