Cloud DDC - Added which claims was matched to the authorization debug view and allowed the claims to be overridden which should make it a lot more useful

diff --git a/Engine/Source/Programs/UnrealCloudDDC/Jupiter.Common/Authentication/Acl.cs b/Engine/Source/Programs/UnrealCloudDDC/Jupiter.Common/Authentication/Acl.cs
--- a/Engine/Source/Programs/UnrealCloudDDC/Jupiter.Common/Authentication/Acl.cs
+++ b/Engine/Source/Programs/UnrealCloudDDC/Jupiter.Common/Authentication/Acl.cs
@@ -142,8 +142,10 @@ public class AclPolicy
 	/// </summary>
 	public AclScope? Scope { get; set; }
 
-	public IEnumerable<JupiterAclAction> Resolve(ClaimsPrincipal user, AccessScope scope, ILogger logger)
+	public IEnumerable<JupiterAclAction> Resolve(ClaimsPrincipal user, AccessScope scope, ILogger logger, out List<string> matchedClaims)
 	{
+		matchedClaims = new List<string>();
+
 		if (Scope == null)
 		{
 			// if no scope is set this policy is not valid so it does nothing
@@ -160,6 +162,7 @@ public class AclPolicy
 			{
 				claimFound = true;
 				logger.LogDebug("Found wildcard claim in policy \'{Name}\'.", Name);
+				matchedClaims.Add("*");
 			}
 			else if (expectedClaim.Contains('=', StringComparison.InvariantCultureIgnoreCase))
 			{
@@ -170,12 +173,14 @@ public class AclPolicy
 				{
 					claimFound = true;
 					logger.LogDebug("Found subset value claim {ClaimName} {ClaimValue} in policy \'{Name}\'.", claimName, claimValue, Name);
+					matchedClaims.Add($"{claimName}={claimValue}");
 				}
 			}
 			else if (user.HasClaim(claim => string.Equals(claim.Type.Trim(), expectedClaim.Trim(), StringComparison.OrdinalIgnoreCase)))
 			{
 				claimFound = true;
 				logger.LogDebug("Found exact claim match {ExpectedClaim} in policy \'{Name}\'.", expectedClaim, Name);
+				matchedClaims.Add(expectedClaim);
 			}
 
 			if (!claimFound)
diff --git a/Engine/Source/Programs/UnrealCloudDDC/Jupiter.Common/Authentication/ScopeAuthorizationHandler.cs b/Engine/Source/Programs/UnrealCloudDDC/Jupiter.Common/Authentication/ScopeAuthorizationHandler.cs
--- a/Engine/Source/Programs/UnrealCloudDDC/Jupiter.Common/Authentication/ScopeAuthorizationHandler.cs
+++ b/Engine/Source/Programs/UnrealCloudDDC/Jupiter.Common/Authentication/ScopeAuthorizationHandler.cs
@@ -59,7 +59,7 @@ namespace Jupiter
 			List<JupiterAclAction> allowedActions = new List<JupiterAclAction>();
 			foreach (AclPolicy acl in _authSettings.CurrentValue.Policies)
 			{
-				allowedActions.AddRange(acl.Resolve(context.User, accessRequest.AccessScope, _logger));
+				allowedActions.AddRange(acl.Resolve(context.User, accessRequest.AccessScope, _logger, out _));
 			}
 
 			if (_authSettings.CurrentValue.UseLegacyConfiguration)
diff --git a/Engine/Source/Programs/UnrealCloudDDC/Jupiter/Controllers/AuthController.cs b/Engine/Source/Programs/UnrealCloudDDC/Jupiter/Controllers/AuthController.cs
--- a/Engine/Source/Programs/UnrealCloudDDC/Jupiter/Controllers/AuthController.cs
+++ b/Engine/Source/Programs/UnrealCloudDDC/Jupiter/Controllers/AuthController.cs
@@ -6,7 +6,9 @@ using System.ComponentModel.DataAnnotations;
 using System.IO;
 using System.Linq;
 using System.Net.Mime;
+using System.Security.Claims;
 using System.Security.Cryptography;
+using System.Text;
 using System.Text.Json;
 using System.Threading.Tasks;
 using EpicGames.Core;
@@ -85,19 +87,15 @@ namespace Jupiter.Controllers
 			
 			CaptureLogger captureLogger = new CaptureLogger();
 			bool policiesFound = false;
-			Dictionary<string, List<JupiterAclAction>> matchedPolicies = new();
-			int suffixCounter = 0;
+			List<MatchedPolicy> matchedPolicies = new();
 			foreach (AclPolicy acl in _authSettings.CurrentValue.Policies)
 			{
 				policiesFound = true;
-				List<JupiterAclAction> grantedActions = acl.Resolve(User, new AccessScope(ns), captureLogger).ToList();
+				List<JupiterAclAction> grantedActions = acl.Resolve(User, new AccessScope(ns), captureLogger, out List<string> matchedClaims).ToList();
 				allowedActions.AddRange(grantedActions);
-				if (!matchedPolicies.TryAdd(acl.Name, grantedActions))
-				{
-					// the policy already existed likely because the name is not globally unique so we just append a random suffix
-					suffixCounter++;
-					matchedPolicies.TryAdd($"{acl.Name}_{suffixCounter}", grantedActions);
-				}
+
+				MatchedPolicy matchedPolicy = new MatchedPolicy { Actions = grantedActions, Claims = matchedClaims };
+				matchedPolicies.Add(matchedPolicy);
 			}
 
 			if (!policiesFound)
@@ -110,26 +108,47 @@ namespace Jupiter.Controllers
 		}
 		
 		[HttpGet("{ns}/{bucket}/actions")]
-		public IActionResult Actions([FromRoute][Required] NamespaceId ns, [FromRoute][Required] BucketId bucket)
+		public IActionResult Actions([FromRoute][Required] NamespaceId ns, [FromRoute][Required] BucketId bucket, [FromQuery] string? claims = null)
 		{
 			List<JupiterAclAction> allowedActions = new List<JupiterAclAction>();
-			
-			Dictionary<string, List<JupiterAclAction>> matchedPolicies = new();
-			int suffixCounter = 0;
+
+			ClaimsPrincipal principal = User;
+			if (claims != null)
+			{
+				string s = Encoding.ASCII.GetString(Convert.FromBase64String(claims));
+				JsonDocument doc = JsonDocument.Parse(s);
+				List<Claim> claimList = new();
+				foreach(JsonProperty prop in doc.RootElement.EnumerateObject())
+				{
+					if (prop.Value.ValueKind == JsonValueKind.String)
+					{
+						claimList.Add(new Claim(prop.Name, prop.Value.GetString()!));
+					} 
+					else if (prop.Value.ValueKind == JsonValueKind.Array)
+					{
+						foreach (JsonElement element in prop.Value.EnumerateArray())
+						{
+							claimList.Add(new Claim(prop.Name, element.GetString()!));
+						}
+					}
+					else
+					{
+						throw new NotImplementedException("Unknown json property type when attempting to convert it into a claim");
+					}
+				}
+				principal = new ClaimsPrincipal(new ClaimsIdentity(claimList));
+			}
+			List<MatchedPolicy> matchedPolicies = new();
 			CaptureLogger captureLogger = new CaptureLogger();
 			bool policiesFound = false;
 			foreach (AclPolicy acl in _authSettings.CurrentValue.Policies)
 			{
 				policiesFound = true;
-				List<JupiterAclAction> grantedActions = acl.Resolve(User, new AccessScope(ns, bucket), captureLogger).ToList();
+				List<JupiterAclAction> grantedActions = acl.Resolve(principal, new AccessScope(ns, bucket), captureLogger, out List<string> matchedClaims).ToList();
 				allowedActions.AddRange(grantedActions);
 
-				if (!matchedPolicies.TryAdd(acl.Name, grantedActions))
-				{
-					// the policy already existed likely because the name is not globally unique so we just append a random suffix
-					suffixCounter++;
-					matchedPolicies.TryAdd($"{acl.Name}_{suffixCounter}", grantedActions);
-				}
+				MatchedPolicy matchedPolicy = new MatchedPolicy { Name = acl.Name, Actions = grantedActions, Claims = matchedClaims };
+				matchedPolicies.Add(matchedPolicy);
 			}
 
 			if (!policiesFound)
@@ -154,6 +173,20 @@ namespace Jupiter.Controllers
 
 		[System.Diagnostics.CodeAnalysis.SuppressMessage("Usage", "CA2227:Collection properties should be read only", Justification = "Used by serialization")]
 		[CbField("matchedPolicies")]
-		public Dictionary<string, List<JupiterAclAction>> MatchedPolicies { get; set; } = new Dictionary<string, List<JupiterAclAction>>();
+		public List<MatchedPolicy> MatchedPolicies { get; set; } = new List<MatchedPolicy>();
+	}
+
+	public class MatchedPolicy
+	{
+		[System.Diagnostics.CodeAnalysis.SuppressMessage("Usage", "CA2227:Collection properties should be read only", Justification = "Used by serialization")]
+		[CbField("actions")] 
+		public List<JupiterAclAction> Actions { get; set; } = new List<JupiterAclAction>();
+
+		[System.Diagnostics.CodeAnalysis.SuppressMessage("Usage", "CA2227:Collection properties should be read only", Justification = "Used by serialization")]
+		[CbField("claims")]
+		public List<string> Claims { get; set; } = new List<string>();
+
+		[CbField("name")]
+		public string Name { get; set; } = null!;
 	}
 }
