Cloud DDC - Tweaked liveness probes to be less aggressive on the restarts of pods (by increasing the timeout). Also added option to run curl healthchecks instead of http get as some kubernetes issues indicates this is more reliable.

diff --git a/Engine/Source/Programs/UnrealCloudDDC/Helm/UnrealCloudDDC/templates/deployment.yaml b/Engine/Source/Programs/UnrealCloudDDC/Helm/UnrealCloudDDC/templates/deployment.yaml
--- a/Engine/Source/Programs/UnrealCloudDDC/Helm/UnrealCloudDDC/templates/deployment.yaml
+++ b/Engine/Source/Programs/UnrealCloudDDC/Helm/UnrealCloudDDC/templates/deployment.yaml
@@ -53,7 +53,7 @@ spec:
       tolerations: {{- include "common.tplvalues.render" (dict "value" .Values.tolerations "context" $) | nindent 8 }}
       {{- end }}
       {{- if or .Values.initContainers .Values.prepareNvmeFilesystem.enabled }}
-      initContainers: 
+      initContainers:
       {{- if .Values.initContainers }}
       {{- include "common.tplvalues.render" (dict "value" .Values.initContainers "context" $) | nindent 6 }}
       {{- end }}
@@ -71,11 +71,11 @@ spec:
         - -c
         - "apt-get update && apt-get -y install mdadm && mdadm --create --verbose /dev/md0 --level=0 --raid-devices={{ $drivesCount }} {{ join " " .Values.prepareNvmeFilesystem.drives }} && mkfs.ext4 -E nodiscard /dev/md0 && mount /dev/md0 /data-nvme || exit 0"
       {{- else }}
-        args: 
+        args:
         - -c
         - mkfs.ext4 -E nodiscard {{ first .Values.prepareNvmeFilesystem.drives }} && mount {{ first .Values.prepareNvmeFilesystem.drives }} /data-nvme || exit 0
       {{- end }}
-      {{- end }}   
+      {{- end }}
         volumeMounts:
         - mountPath: /data-nvme
           name: data
@@ -112,9 +112,14 @@ spec:
           containerPort: 8083
         {{- if .Values.livenessProbe.enabled }}
         livenessProbe:
+          {{- if .Values.livenessProbe.useCurlHealthCheck }}
+          exec:
+            command: [ "sh", "-c", "curl --fail localhost:80/{{ default "health/live" .Values.livenessProbe.path }}"]
+          {{- else }}
           httpGet:
             path: {{ default "health/live" .Values.livenessProbe.path }}
             port: http
+          {{- end }}
           initialDelaySeconds: {{ .Values.livenessProbe.initialDelaySeconds }}
           periodSeconds: {{ .Values.livenessProbe.periodSeconds }}
           timeoutSeconds: {{ .Values.livenessProbe.timeoutSeconds }}
@@ -123,9 +128,14 @@ spec:
         {{- end }}
         {{- if .Values.readinessProbe.enabled }}
         readinessProbe:
+          {{- if .Values.readinessProbe.useCurlHealthCheck }}
+          exec:
+            command: [ "sh", "-c", "curl --fail localhost:80/{{ default "health/ready" .Values.readinessProbe.path }}"]
+          {{- else }}
           httpGet:
-            path:  {{ default "health/ready" .Values.readinessProbe.path }}
+            path: {{ default "health/ready" .Values.readinessProbe.path }}
             port: http
+          {{- end }}
           initialDelaySeconds: {{ .Values.readinessProbe.initialDelaySeconds }}
           periodSeconds: {{ .Values.readinessProbe.periodSeconds }}
           timeoutSeconds: {{ .Values.readinessProbe.timeoutSeconds }}
@@ -136,7 +146,7 @@ spec:
         lifecycle:
           preStop:
             exec:
-              command: 
+              command:
               # Load balancers can be slow to deregister instances, this can be used to introduce a delay to provide zero downtime
               - "sleep"
               - "{{ .Values.preStopHookDelay }}"
@@ -189,7 +199,7 @@ spec:
     {{- if .Values.nginx.enabled }}
       - name: nginx
         image: nginx:1.25.4
-        {{- if .Values.nginx.env }} 
+        {{- if .Values.nginx.env }}
         env:
         {{- include "common.tplvalues.render" ( dict "value" .Values.nginx.env "context" $ ) | nindent 8 }}
         {{- end }}
@@ -231,7 +241,7 @@ spec:
         lifecycle:
           preStop:
             exec:
-              command: 
+              command:
               # Load balancers can be slow to deregister instances, this can be used to introduce a delay to provide zero downtime
               - "sleep"
               - "{{ .Values.preStopHookDelay }}"
