Cloud DDC - Deleting namespaces and buckets that are found to exist but not be used, added feature toggle to disable this in case we see any issues with it.

diff --git a/Engine/Source/Programs/UnrealCloudDDC/Jupiter/Implementation/Objects/ScyllaReferencesStore.cs b/Engine/Source/Programs/UnrealCloudDDC/Jupiter/Implementation/Objects/ScyllaReferencesStore.cs
--- a/Engine/Source/Programs/UnrealCloudDDC/Jupiter/Implementation/Objects/ScyllaReferencesStore.cs
+++ b/Engine/Source/Programs/UnrealCloudDDC/Jupiter/Implementation/Objects/ScyllaReferencesStore.cs
@@ -652,16 +652,21 @@ namespace Jupiter.Implementation
 					}
 					_logger.LogWarning("Namespace {Namespace} Bucket {Bucket} has no refs. Should be removed from tracking.", ns, bucket);
 
-					// TODO: introduce option to actually delete this and not just report
-					//await _session.ExecuteAsync(new SimpleStatement("DELETE FROM buckets_v2 WHERE namespace = ? and bucket = ?", ns.ToString(), bucket.ToString()));
+					if (_settings.CurrentValue.CleanupOldTrackedState)
+					{
+						await _session.ExecuteAsync(new SimpleStatement("DELETE FROM buckets_v2 WHERE namespace = ? and bucket = ?", ns.ToString(), bucket.ToString()));
+					}
 				}
 
 				if (!refFound)
 				{
 					_logger.LogWarning("Namespace {Namespace} has no refs. Should be removed from tracking.", ns);
 
-					// TODO: introduce option to actually delete this and not just report
-					// await _session.ExecuteAsync(new SimpleStatement("DELETE FROM buckets_v2 WHERE namespace = ?", ns.ToString()));
+
+					if (_settings.CurrentValue.CleanupOldTrackedState)
+					{
+						await _session.ExecuteAsync(new SimpleStatement("DELETE FROM buckets_v2 WHERE namespace = ?", ns.ToString()));
+					}
 				}
 			}
 
diff --git a/Engine/Source/Programs/UnrealCloudDDC/Jupiter/Implementation/Objects/ScyllaSettings.cs b/Engine/Source/Programs/UnrealCloudDDC/Jupiter/Implementation/Objects/ScyllaSettings.cs
--- a/Engine/Source/Programs/UnrealCloudDDC/Jupiter/Implementation/Objects/ScyllaSettings.cs
+++ b/Engine/Source/Programs/UnrealCloudDDC/Jupiter/Implementation/Objects/ScyllaSettings.cs
@@ -115,6 +115,11 @@ public class ScyllaSettings : IValidatableObject
 	/// </summary>
 	public bool AllowParallelRecordFetch { get; set; } = false;
 
+	/// <summary>
+	/// Can be used to disable the cleanup of unused namespaces and buckets
+	/// </summary>
+	public bool CleanupOldTrackedState { get; set; } = true;
+
 	public IEnumerable<ValidationResult> Validate(ValidationContext validationContext)
 	{
 		List<ValidationResult> results = new List<ValidationResult>();
