Cloud DDC - Added cleanup of incorrect blob reference entires when detected during the consistency check step

diff --git a/Engine/Source/Programs/UnrealCloudDDC/Jupiter/Implementation/Blob/BlobIndexConsistencyCheckService.cs b/Engine/Source/Programs/UnrealCloudDDC/Jupiter/Implementation/Blob/BlobIndexConsistencyCheckService.cs
--- a/Engine/Source/Programs/UnrealCloudDDC/Jupiter/Implementation/Blob/BlobIndexConsistencyCheckService.cs
+++ b/Engine/Source/Programs/UnrealCloudDDC/Jupiter/Implementation/Blob/BlobIndexConsistencyCheckService.cs
@@ -101,7 +101,8 @@ namespace Jupiter.Implementation
 						if (!exists)
 						{
 							Interlocked.Increment(ref countOfIncorrectBlobsFound);
-							_logger.LogWarning("Blob {Blob} in namespace {Namespace} had entries in the ref table but does not exist.", blobIdentifier, ns);
+							_logger.LogWarning("Blob {Blob} in namespace {Namespace} had entries in the ref table but does not exist. Cleaning up.", blobIdentifier, ns);
+							await _blobIndex.RemoveReferencesAsync(ns, blobIdentifier, null, cancellationToken);
 						}
 					}
 					catch (Exception e)
