Cloud DDC - Added consistency check for partial objects (refs) as well as cleanup of them. Can happen when using TTL on refs (which has not been enabled by default) due to a misconfiguration of scylla when updating the is_finalized column which is also fixed in this change.

diff --git a/Engine/Source/Programs/UnrealCloudDDC/Jupiter/Implementation/Blob/RefStoreConsistencyCheckService.cs b/Engine/Source/Programs/UnrealCloudDDC/Jupiter/Implementation/Blob/RefStoreConsistencyCheckService.cs
--- a/Engine/Source/Programs/UnrealCloudDDC/Jupiter/Implementation/Blob/RefStoreConsistencyCheckService.cs
+++ b/Engine/Source/Programs/UnrealCloudDDC/Jupiter/Implementation/Blob/RefStoreConsistencyCheckService.cs
@@ -107,6 +107,19 @@ namespace Jupiter.Implementation
 						}
 					}
 				}
+				if (_settings.CurrentValue.CheckRefStorePartialObjects)
+				{
+					// check if any ref can be found by listing but not when explicitly looked up, indicates it doesn't contain all the expected fields and thus is a partial record that should be removed to fix bugs with how TTL was set when is_finalized was updated
+					try
+					{
+						await _referencesStore.GetAsync(ns, bucket, refId, IReferencesStore.FieldFlags.None, IReferencesStore.OperationFlags.BypassCache, cancellationToken);
+					}
+					catch (RefNotFoundException)
+					{
+						_logger.LogWarning("Partial record {Namespace} {Bucket} {Name} found. Deleting.", ns, bucket, refId);
+						await _referencesStore.DeleteAsync(ns, bucket, refId, cancellationToken);
+					}
+				}
 				Interlocked.Increment(ref countOfRefsChecked);
 			}
 
diff --git a/Engine/Source/Programs/UnrealCloudDDC/Jupiter/Implementation/Objects/ScyllaReferencesStore.cs b/Engine/Source/Programs/UnrealCloudDDC/Jupiter/Implementation/Objects/ScyllaReferencesStore.cs
--- a/Engine/Source/Programs/UnrealCloudDDC/Jupiter/Implementation/Objects/ScyllaReferencesStore.cs
+++ b/Engine/Source/Programs/UnrealCloudDDC/Jupiter/Implementation/Objects/ScyllaReferencesStore.cs
@@ -204,7 +204,22 @@ namespace Jupiter.Implementation
 		{
 			using TelemetrySpan scope = _tracer.BuildScyllaSpan("scylla.finalize").SetAttribute("resource.name", $"{ns}.{bucket}.{name}");
 
-			await _mapper.UpdateAsync<ScyllaObject>("SET is_finalized=true WHERE namespace=? AND bucket=? AND name=?", ns.ToString(), bucket.ToString(), name.ToString());
+			int? ttl = null;
+			NamespacePolicy policy = _namespacePolicyResolver.GetPoliciesForNs(ns);
+			NamespacePolicy.StoragePoolGCMethod gcMethod = policy.GcMethod ?? NamespacePolicy.StoragePoolGCMethod.LastAccess;
+			if (gcMethod == NamespacePolicy.StoragePoolGCMethod.TTL)
+			{
+				ttl = (int)policy.DefaultTTL.TotalSeconds;
+			}
+
+			if (ttl != null)
+			{
+				await _mapper.UpdateAsync<ScyllaObject>("USING TTL ? SET is_finalized=true WHERE namespace=? AND bucket=? AND name=?", ttl.Value, ns.ToString(), bucket.ToString(), name.ToString());
+			}
+			else
+			{
+				await _mapper.UpdateAsync<ScyllaObject>("SET is_finalized=true WHERE namespace=? AND bucket=? AND name=?", ns.ToString(), bucket.ToString(), name.ToString());
+			}
 		}
 
 		public async Task<DateTime?> GetLastAccessTimeAsync(NamespaceId ns, BucketId bucket, RefId key, CancellationToken cancellationToken)
@@ -662,7 +677,6 @@ namespace Jupiter.Implementation
 				{
 					_logger.LogWarning("Namespace {Namespace} has no refs. Should be removed from tracking.", ns);
 
-
 					if (_settings.CurrentValue.CleanupOldTrackedState)
 					{
 						await _session.ExecuteAsync(new SimpleStatement("DELETE FROM buckets_v2 WHERE namespace = ?", ns.ToString()));
diff --git a/Engine/Source/Programs/UnrealCloudDDC/Jupiter/Settings.cs b/Engine/Source/Programs/UnrealCloudDDC/Jupiter/Settings.cs
--- a/Engine/Source/Programs/UnrealCloudDDC/Jupiter/Settings.cs
+++ b/Engine/Source/Programs/UnrealCloudDDC/Jupiter/Settings.cs
@@ -295,6 +295,7 @@ namespace Jupiter
 		public bool RunBlobStoreConsistencyCheckOnRootStore { get; set; } = false;
 		public bool CheckRefStoreLastAccessTimeConsistency { get; set; } = false;
 		public bool CheckRefStoreRegionalConsistency { get; set; } = false;
+		public bool CheckRefStorePartialObjects { get; set; } = false;
 		public string[] RegionalConsistencyCheckNamespaces { get; set; } = Array.Empty<string>();
 		public bool EnableBlobReferenceChecks { get; set; } = false;
 	}
