Cloud DDC - Fixed issue were blob index consistency checks were not being run

diff --git a/Engine/Source/Programs/UnrealCloudDDC/Jupiter/Implementation/Blob/BlobIndexConsistencyCheckService.cs b/Engine/Source/Programs/UnrealCloudDDC/Jupiter/Implementation/Blob/BlobIndexConsistencyCheckService.cs
--- a/Engine/Source/Programs/UnrealCloudDDC/Jupiter/Implementation/Blob/BlobIndexConsistencyCheckService.cs
+++ b/Engine/Source/Programs/UnrealCloudDDC/Jupiter/Implementation/Blob/BlobIndexConsistencyCheckService.cs
@@ -59,11 +59,11 @@ namespace Jupiter.Implementation
 				return false;
 			}
 
-			if (!_settings.CurrentValue.EnableBlobReferenceChecks)
+			if (_settings.CurrentValue.EnableBlobReferenceChecks)
 			{
 				await RunConsistencyCheckBlobReferencesAsync(cancellationToken);
 			}
-			if (!_settings.CurrentValue.EnableBlobIndexChecks)
+			if (_settings.CurrentValue.EnableBlobIndexChecks)
 			{
 				await RunConsistencyCheckAsync(cancellationToken);
 			}
