Cloud DDC - Fixed issue were remove references in the blob index was not correctly deleting all reference tracking causing db bloat

diff --git a/Engine/Source/Programs/UnrealCloudDDC/Jupiter/Implementation/Blob/ScyllaBlobIndex.cs b/Engine/Source/Programs/UnrealCloudDDC/Jupiter/Implementation/Blob/ScyllaBlobIndex.cs
--- a/Engine/Source/Programs/UnrealCloudDDC/Jupiter/Implementation/Blob/ScyllaBlobIndex.cs
+++ b/Engine/Source/Programs/UnrealCloudDDC/Jupiter/Implementation/Blob/ScyllaBlobIndex.cs
@@ -340,7 +340,7 @@ public class ScyllaBlobIndex : IBlobIndex
 
 		if (referencesToRemove == null)
 		{
-			await _mapper.DeleteAsync<ScyllaBlobIncomingReference>("WHERE namespace = ? AND blob_id = ?", nsAsString, id.ToString());
+			await _mapper.DeleteAsync<ScyllaBlobIncomingReference>("WHERE namespace = ? AND blob_id = ?", nsAsString, id.HashData);
 		}
 		else
 		{
