Cloud DDC - Added consistency check service for builds to detect partial build records (left over due to a bug in cleaning up TTLed records)

diff --git a/Engine/Source/Programs/UnrealCloudDDC/Jupiter/Implementation/Builds/BuildStoreConsistencyCheckService.cs b/Engine/Source/Programs/UnrealCloudDDC/Jupiter/Implementation/Builds/BuildStoreConsistencyCheckService.cs
new file mode 100644
--- /dev/null
+++ b/Engine/Source/Programs/UnrealCloudDDC/Jupiter/Implementation/Builds/BuildStoreConsistencyCheckService.cs
@@ -0,0 +1,103 @@
+// Copyright Epic Games, Inc. All Rights Reserved.
+
+using System;
+using System.Threading;
+using System.Threading.Tasks;
+using EpicGames.Horde.Storage;
+using EpicGames.Serialization;
+using Microsoft.Extensions.Logging;
+using Microsoft.Extensions.Options;
+using OpenTelemetry.Trace;
+
+namespace Jupiter.Implementation
+{
+	public class BuildStoreConsistencyState
+	{
+	}
+
+	// ReSharper disable once ClassNeverInstantiated.Global
+	public class BuildStoreConsistencyCheckService : PollingService<BuildStoreConsistencyState>
+	{
+		private readonly IOptionsMonitor<ConsistencyCheckSettings> _settings;
+		private readonly ILeaderElection _leaderElection;
+		private readonly IBuildStore _buildStore;
+		private readonly Tracer _tracer;
+		private readonly ILogger _logger;
+
+		protected override bool ShouldStartPolling()
+		{
+			return _settings.CurrentValue.EnableBuildStoreConsistencyCheck;
+		}
+
+		public BuildStoreConsistencyCheckService(IOptionsMonitor<ConsistencyCheckSettings> settings, ILeaderElection leaderElection, IBuildStore buildStore, Tracer tracer, ILogger<BuildStoreConsistencyCheckService> logger) : base(serviceName: nameof(BuildStoreConsistencyCheckService), TimeSpan.FromSeconds(settings.CurrentValue.ConsistencyCheckPollFrequencySeconds), new BuildStoreConsistencyState(), logger)
+		{
+			_settings = settings;
+			_leaderElection = leaderElection;
+			_buildStore = buildStore;
+			_tracer = tracer;
+			_logger = logger;
+		}
+
+		public override async Task<bool> OnPollAsync(BuildStoreConsistencyState state, CancellationToken cancellationToken)
+		{
+			bool shouldEnable = _settings.CurrentValue.EnableBuildStoreConsistencyCheck;
+			if (!shouldEnable)
+			{
+				_logger.LogInformation("Skipped running build store consistency check as it is disabled");
+				return false;
+			}
+
+			if (!_leaderElection.IsThisInstanceLeader())
+			{
+				_logger.LogInformation("Skipped running build store consistency check because this instance was not the leader");
+				return false;
+			}
+
+			await RunConsistencyCheckAsync(cancellationToken);
+
+			return true;
+		}
+		
+		private async Task RunConsistencyCheckAsync(CancellationToken cancellationToken)
+		{
+			ulong countOfBuildsChecked = 0;
+			ulong countOfIncorrectBuildsFound = 0;
+			await Parallel.ForEachAsync(_buildStore.ListAllBuildsAsync(cancellationToken), new ParallelOptions
+			{
+				CancellationToken = cancellationToken,
+				MaxDegreeOfParallelism = _settings.CurrentValue.BlobIndexMaxParallelOperations,
+			},
+				async (tuple, token) =>
+				{
+					(NamespaceId ns, BucketId bucket, CbObjectId buildId) = tuple;
+					Interlocked.Increment(ref countOfBuildsChecked);
+
+					if (countOfBuildsChecked % 100 == 0)
+					{
+						_logger.LogInformation("Consistency check running on build store, count of builds processed so far: {CountOfBuilds}", countOfBuildsChecked);
+					}
+
+					using TelemetrySpan scope = _tracer.StartActiveSpan("consistency_check.build_store").SetAttribute("resource.name", $"{ns}.{bucket}.{buildId}").SetAttribute("operation.name", "consistency_check.build_store");
+
+					BuildRecord? record = await _buildStore.GetBuildAsync(ns, bucket, buildId);
+					if (record == null)
+					{
+						_logger.LogWarning("Partial build {BuildId} in bucket {Bucket} and namespace {Namespace} .", buildId, bucket, ns);
+					}
+
+					if (_settings.CurrentValue.AllowDeletesInBuildStore)
+					{
+						await _buildStore.DeleteBuild(ns, bucket, buildId);
+					}
+				}
+			);
+
+			_logger.LogInformation("Build Store Consistency check finished, found {CountOfIncorrectBuilds} incorrect builds. Processed {CountOfBuilds} builds.", countOfIncorrectBuildsFound, countOfBuildsChecked);
+		}
+
+		protected override Task OnStopping(BuildStoreConsistencyState state)
+		{
+			return Task.CompletedTask;
+		}
+	}
+}
diff --git a/Engine/Source/Programs/UnrealCloudDDC/Jupiter/Implementation/Builds/IBuildStore.cs b/Engine/Source/Programs/UnrealCloudDDC/Jupiter/Implementation/Builds/IBuildStore.cs
--- a/Engine/Source/Programs/UnrealCloudDDC/Jupiter/Implementation/Builds/IBuildStore.cs
+++ b/Engine/Source/Programs/UnrealCloudDDC/Jupiter/Implementation/Builds/IBuildStore.cs
@@ -1,6 +1,7 @@
 // Copyright Epic Games, Inc. All Rights Reserved.
 
 using System.Collections.Generic;
+using System.Threading;
 using System.Threading.Tasks;
 using EpicGames.Horde.Storage;
 using EpicGames.Serialization;
@@ -19,6 +20,7 @@ public interface IBuildStore
 	Task UpdateTTL(NamespaceId ns, BucketId bucket, CbObjectId buildId, uint ttl);
 	Task DeleteBuild(NamespaceId ns, BucketId bucket, CbObjectId buildId);
 	Task<uint?> GetTTL(NamespaceId ns, BucketId bucket, CbObjectId buildId);
+	IAsyncEnumerable<(NamespaceId, BucketId, CbObjectId)> ListAllBuildsAsync(CancellationToken cancellationToken);
 }
 
 public class BuildMetadata
diff --git a/Engine/Source/Programs/UnrealCloudDDC/Jupiter/Implementation/Builds/MemoryBuildStore.cs b/Engine/Source/Programs/UnrealCloudDDC/Jupiter/Implementation/Builds/MemoryBuildStore.cs
--- a/Engine/Source/Programs/UnrealCloudDDC/Jupiter/Implementation/Builds/MemoryBuildStore.cs
+++ b/Engine/Source/Programs/UnrealCloudDDC/Jupiter/Implementation/Builds/MemoryBuildStore.cs
@@ -4,6 +4,8 @@ using System;
 using System.Collections.Concurrent;
 using System.Collections.Generic;
 using System.Linq;
+using System.Runtime.CompilerServices;
+using System.Threading;
 using System.Threading.Tasks;
 using EpicGames.Horde.Storage;
 using EpicGames.Serialization;
@@ -119,6 +121,15 @@ public class MemoryBuildStore : IBuildStore
 		return Task.FromResult(_builds.TryGetValue(key, out MemoryBuildRecord? value) ? (uint?)(value.ExpireTime - DateTime.Now).TotalSeconds : null);
 	}
 
+	public async IAsyncEnumerable<(NamespaceId, BucketId, CbObjectId)> ListAllBuildsAsync([EnumeratorCancellation] CancellationToken cancellationToken)
+	{
+		await Task.CompletedTask;
+		foreach (MemoryBuildRecord build in _builds.Values.OrderByDescending(record => record.CreationDate))
+		{
+			yield return (build.Namespace, build.Bucket, build.BuildId);
+		}
+	}
+
 	private static string BuildKey(NamespaceId ns, BucketId bucket, CbObjectId buildId)
 	{
 		return $"{ns}.{bucket}.{buildId}";
diff --git a/Engine/Source/Programs/UnrealCloudDDC/Jupiter/Implementation/Builds/ScyllaBuildStore.cs b/Engine/Source/Programs/UnrealCloudDDC/Jupiter/Implementation/Builds/ScyllaBuildStore.cs
--- a/Engine/Source/Programs/UnrealCloudDDC/Jupiter/Implementation/Builds/ScyllaBuildStore.cs
+++ b/Engine/Source/Programs/UnrealCloudDDC/Jupiter/Implementation/Builds/ScyllaBuildStore.cs
@@ -2,6 +2,8 @@
 
 using System;
 using System.Collections.Generic;
+using System.Runtime.CompilerServices;
+using System.Threading;
 using System.Threading.Tasks;
 using Cassandra;
 using Cassandra.Mapping;
@@ -14,6 +16,7 @@ namespace Jupiter.Implementation;
 
 public class ScyllaBuildStore : IBuildStore
 {
+	private readonly IScyllaSessionManager _scyllaSessionManager;
 	private readonly Tracer _tracer;
 	private readonly ISession _session;
 	private readonly Mapper _mapper;
@@ -21,6 +24,7 @@ public class ScyllaBuildStore : IBuildStore
 
 	public ScyllaBuildStore(IScyllaSessionManager scyllaSessionManager, IOptionsMonitor<ScyllaSettings> scyllaSettings, Tracer tracer)
 	{
+		_scyllaSessionManager = scyllaSessionManager;
 		_tracer = tracer;
 		_session = scyllaSessionManager.GetSessionForReplicatedKeyspace();
 		_mapper = new Mapper(_session);
@@ -132,6 +136,18 @@ public class ScyllaBuildStore : IBuildStore
 		return null;
 	}
 
+	public async IAsyncEnumerable<(NamespaceId, BucketId, CbObjectId)> ListAllBuildsAsync([EnumeratorCancellation] CancellationToken cancellationToken)
+	{
+		using TelemetrySpan scope = _tracer.BuildScyllaSpan("scylla.get_all_builds");
+
+		string cqlOptions = _scyllaSessionManager.IsScylla ? "BYPASS CACHE" : "";
+
+		foreach (ScyllaBuildEntry buildEntry in await _mapper.FetchAsync<ScyllaBuildEntry>(cqlOptions))
+		{
+			yield return (new NamespaceId(buildEntry.Namespace), new BucketId(buildEntry.BucketId), new CbObjectId(buildEntry.BuildId));
+		}
+	}
+
 	public async Task UpdateTTL(NamespaceId ns, BucketId bucket, CbObjectId buildId, uint ttl)
 	{
 		ScyllaBuildEntry? buildEntry = await _mapper.SingleOrDefaultAsync<ScyllaBuildEntry>("WHERE namespace = ? AND bucket_id = ? AND build_id = ?", ns.ToString(), bucket.ToString(), buildId.ToByteArray());
diff --git a/Engine/Source/Programs/UnrealCloudDDC/Jupiter/Jupiter.Startup.cs b/Engine/Source/Programs/UnrealCloudDDC/Jupiter/Jupiter.Startup.cs
--- a/Engine/Source/Programs/UnrealCloudDDC/Jupiter/Jupiter.Startup.cs
+++ b/Engine/Source/Programs/UnrealCloudDDC/Jupiter/Jupiter.Startup.cs
@@ -196,6 +196,9 @@ namespace Jupiter
 			services.AddSingleton<RefStoreConsistencyCheckService>();
 			services.AddHostedService<RefStoreConsistencyCheckService>(p => p.GetService<RefStoreConsistencyCheckService>()!);
 
+			services.AddSingleton<BuildStoreConsistencyCheckService>();
+			services.AddHostedService<BuildStoreConsistencyCheckService>(p => p.GetService<BuildStoreConsistencyCheckService>()!);
+			
 			services.AddSingleton<MetricsService>();
 			services.AddHostedService<MetricsService>(p => p.GetService<MetricsService>()!);
 
diff --git a/Engine/Source/Programs/UnrealCloudDDC/Jupiter/Settings.cs b/Engine/Source/Programs/UnrealCloudDDC/Jupiter/Settings.cs
--- a/Engine/Source/Programs/UnrealCloudDDC/Jupiter/Settings.cs
+++ b/Engine/Source/Programs/UnrealCloudDDC/Jupiter/Settings.cs
@@ -298,5 +298,7 @@ namespace Jupiter
 		public bool CheckRefStorePartialObjects { get; set; } = false;
 		public string[] RegionalConsistencyCheckNamespaces { get; set; } = Array.Empty<string>();
 		public bool EnableBlobReferenceChecks { get; set; } = false;
+		public bool EnableBuildStoreConsistencyCheck { get; set; } = false;
+		public bool AllowDeletesInBuildStore { get; set; } = false;
 	}
 }
