Cloud DDC - Fixed issue with build store cleanup not tracking number of builds correctly and an issue were it would cleanup more builds then intended.

diff --git a/Engine/Source/Programs/UnrealCloudDDC/Jupiter/Implementation/Builds/BuildStoreConsistencyCheckService.cs b/Engine/Source/Programs/UnrealCloudDDC/Jupiter/Implementation/Builds/BuildStoreConsistencyCheckService.cs
--- a/Engine/Source/Programs/UnrealCloudDDC/Jupiter/Implementation/Builds/BuildStoreConsistencyCheckService.cs
+++ b/Engine/Source/Programs/UnrealCloudDDC/Jupiter/Implementation/Builds/BuildStoreConsistencyCheckService.cs
@@ -83,11 +83,12 @@ namespace Jupiter.Implementation
 					if (record == null)
 					{
 						_logger.LogWarning("Partial build {BuildId} in bucket {Bucket} and namespace {Namespace} .", buildId, bucket, ns);
-					}
 
-					if (_settings.CurrentValue.AllowDeletesInBuildStore)
-					{
-						await _buildStore.DeleteBuild(ns, bucket, buildId);
+						Interlocked.Increment(ref countOfIncorrectBuildsFound);
+						if (_settings.CurrentValue.AllowDeletesInBuildStore)
+						{
+							await _buildStore.DeleteBuild(ns, bucket, buildId);
+						}
 					}
 				}
 			);
