Cloud DDC - Handle task cancellation in search, its expected that this could happen for very large queries and we should just suppress the exception

diff --git a/Engine/Source/Programs/UnrealCloudDDC/Jupiter/Controllers/BuildsController.cs b/Engine/Source/Programs/UnrealCloudDDC/Jupiter/Controllers/BuildsController.cs
--- a/Engine/Source/Programs/UnrealCloudDDC/Jupiter/Controllers/BuildsController.cs
+++ b/Engine/Source/Programs/UnrealCloudDDC/Jupiter/Controllers/BuildsController.cs
@@ -477,63 +477,71 @@ namespace Jupiter.Controllers
 			SortedList<CbObjectId, BuildMetadata> foundResults = new (new CbObjectIdReverseComparer());
 
 			using CancellationTokenSource cancellationTokenSource = new CancellationTokenSource();
-			await Parallel.ForEachAsync(bucketsToSearch, cancellationTokenSource.Token, async (bucket, token) =>
+			try
 			{
-				LogAuditEntry(HttpMethod.Post, nameof(SearchBuildAsync), ns, bucket, null, null);
-				
-				await foreach (BuildMetadata build in _buildStore.ListBuildsAsync(ns, bucket, includeTTL).WithCancellation(token))
+				await Parallel.ForEachAsync(bucketsToSearch, cancellationTokenSource.Token, async (bucket, token) =>
 				{
-					if (!build.IsFinalized)
-					{
-						// we ignore non-finalized builds
-						continue;
-					}
-
-					if (token.IsCancellationRequested)
+					LogAuditEntry(HttpMethod.Post, nameof(SearchBuildAsync), ns, bucket, null, null);
+				
+					await foreach (BuildMetadata build in _buildStore.ListBuildsAsync(ns, bucket, includeTTL).WithCancellation(token))
 					{
-						partialResult = true;
-						break;
-					}
-					int count = Interlocked.Increment(ref countConsidered);
+						if (!build.IsFinalized)
+						{
+							// we ignore non-finalized builds
+							continue;
+						}
 
-					if (count >= searchRequest.Options.Max)
-					{
-						// we have considered all the builds we are allowed to consider
-						partialResult = true;
-						await cancellationTokenSource.CancelAsync();
-						break;
-					}
+						if (token.IsCancellationRequested)
+						{
+							partialResult = true;
+							break;
+						}
+						int count = Interlocked.Increment(ref countConsidered);
 
-					timeSpent = DateTime.Now - startTime;
-					if (checkForMaxTime && timeSpent > maxTimeToSpend)
-					{
-						// spent all the time we were given for this query
-						timedOut = true;
-						partialResult = true;
-						await cancellationTokenSource.CancelAsync();
-						break;
-					}
+						if (count >= searchRequest.Options.Max)
+						{
+							// we have considered all the builds we are allowed to consider
+							partialResult = true;
+							await cancellationTokenSource.CancelAsync();
+							break;
+						}
 
-					bool match = true;
-					foreach (ISearchOp op in searchOps)
-					{
-						CbObject o = build.Metadata;
-						if (!op.Matches(build.BuildId, o, _logger))
+						timeSpent = DateTime.Now - startTime;
+						if (checkForMaxTime && timeSpent > maxTimeToSpend)
 						{
-							match = false;
+							// spent all the time we were given for this query
+							timedOut = true;
+							partialResult = true;
+							await cancellationTokenSource.CancelAsync();
 							break;
 						}
-					}
 
-					if (match)
-					{
-						lock (foundResults)
+						bool match = true;
+						foreach (ISearchOp op in searchOps)
 						{
-							foundResults.Add(build.BuildId, build);
+							CbObject o = build.Metadata;
+							if (!op.Matches(build.BuildId, o, _logger))
+							{
+								match = false;
+								break;
+							}
+						}
+
+						if (match)
+						{
+							lock (foundResults)
+							{
+								foundResults.Add(build.BuildId, build);
+							}
 						}
 					}
-				}
-			});
+				});
+			}
+			catch (TaskCanceledException)
+			{
+				// cancellation is not an issue so we can suppress this
+				partialResult = true;
+			}
 			// Collect the results that was newest that match our limits and search
 			List<BuildMetadata> foundResultsInBuckets = foundResults.Values.Skip(searchRequest.Options.Skip).Take(searchRequest.Options.Limit).ToList();
 
