Cloud DDC - Disabled scylla tracing by default as its very large volume ,can be renabled using the OTEL_TRACE_SCYLLA_REQUESTS enviroment variable.

diff --git a/Engine/Source/Programs/UnrealCloudDDC/Jupiter.Common/BaseStartup.cs b/Engine/Source/Programs/UnrealCloudDDC/Jupiter.Common/BaseStartup.cs
--- a/Engine/Source/Programs/UnrealCloudDDC/Jupiter.Common/BaseStartup.cs
+++ b/Engine/Source/Programs/UnrealCloudDDC/Jupiter.Common/BaseStartup.cs
@@ -295,6 +295,9 @@ namespace Jupiter
 			string? useConsoleExporterString = Configuration["OTEL_USE_CONSOLE_EXPORTER"];
 			_ = bool.TryParse(useConsoleExporterString, out bool useConsoleExporter);
 
+			string? traceScyllaRequestsString = Configuration["OTEL_TRACE_SCYLLA_REQUESTS"];
+			_ = bool.TryParse(traceScyllaRequestsString, out bool traceScyllaRequests);
+
 			services.AddOpenTelemetry().ConfigureResource(builder =>
 			{
 				builder.AddService("UnrealCloudDDC", serviceNamespace: "Jupiter", serviceVersion: otelServiceVersion)
@@ -365,11 +368,11 @@ namespace Jupiter
 				{
 					builder.AddConsoleExporter();
 				}
-				builder.AddSource("UnrealCloudDDC", "ScyllaDB");
+				builder.AddSource("UnrealCloudDDC", traceScyllaRequests ? "ScyllaDB" : "");
 			}).WithMetrics(builder =>
 			{
 				builder
-					.AddMeter("UnrealCloudDDC", "ScyllaDB")
+					.AddMeter("UnrealCloudDDC", traceScyllaRequests ? "ScyllaDB" : "")
 					.AddOtlpExporter()
 					.AddAspNetCoreInstrumentation()
 					.AddHttpClientInstrumentation()
