Cloud DDC - Run multiple ref store consistency checks in parallell to try and speed it up

diff --git a/Engine/Source/Programs/UnrealCloudDDC/Jupiter/Implementation/Blob/RefStoreConsistencyCheckService.cs b/Engine/Source/Programs/UnrealCloudDDC/Jupiter/Implementation/Blob/RefStoreConsistencyCheckService.cs
--- a/Engine/Source/Programs/UnrealCloudDDC/Jupiter/Implementation/Blob/RefStoreConsistencyCheckService.cs
+++ b/Engine/Source/Programs/UnrealCloudDDC/Jupiter/Implementation/Blob/RefStoreConsistencyCheckService.cs
@@ -70,58 +70,60 @@ namespace Jupiter.Implementation
 			ulong countOfRefsChecked = 0;
 			ulong countOfMissingLastAccessTime = 0;
 			ulong countOfRegionalInconsistentRefs = 0;
-			await foreach ((NamespaceId ns, BucketId bucket, RefId refId) in _referencesStore.GetRecordsWithoutAccessTimeAsync(cancellationToken))
-			{
-				using TelemetrySpan scope = _tracer.StartActiveSpan("consistency_check.ref_store")
+			await Parallel.ForEachAsync(_referencesStore.GetRecordsWithoutAccessTimeAsync(cancellationToken), cancellationToken,
+				async (record, token) =>
+				{
+					(NamespaceId ns, BucketId bucket, RefId refId) = record;
+					using TelemetrySpan scope = _tracer.StartActiveSpan("consistency_check.ref_store")
 					.SetAttribute("operation.name", "consistency_check.ref_store")
 					.SetAttribute("resource.name", $"{ns}.{bucket}.{refId}");
 
-				if (_settings.CurrentValue.CheckRefStoreLastAccessTimeConsistency)
-				{
-					DateTime? lastAccessTime = await _referencesStore.GetLastAccessTimeAsync(ns, bucket, refId, cancellationToken);
-					if (!lastAccessTime.HasValue)
+					if (_settings.CurrentValue.CheckRefStoreLastAccessTimeConsistency)
 					{
-						// if there is no last access time record we add one so that the two tables are consistent, this will make the ref record be considered by the GC and thus cleanup anything that might be very old (but its added as a new record so it will take until the configured cleanup time has passed)
-						await _referencesStore.UpdateLastAccessTimeAsync(ns, bucket, refId, DateTime.Now, cancellationToken);
+						DateTime? lastAccessTime = await _referencesStore.GetLastAccessTimeAsync(ns, bucket, refId, token);
+						if (!lastAccessTime.HasValue)
+						{
+							// if there is no last access time record we add one so that the two tables are consistent, this will make the ref record be considered by the GC and thus cleanup anything that might be very old (but its added as a new record so it will take until the configured cleanup time has passed)
+							await _referencesStore.UpdateLastAccessTimeAsync(ns, bucket, refId, DateTime.Now, token);
 
-						Interlocked.Increment(ref countOfMissingLastAccessTime);
+							Interlocked.Increment(ref countOfMissingLastAccessTime);
+						}
 					}
-				}
-				if (_settings.CurrentValue.CheckRefStoreRegionalConsistency)
-				{
-					if (_settings.CurrentValue.RegionalConsistencyCheckNamespaces.Contains(ns.ToString()))
+					if (_settings.CurrentValue.CheckRefStoreRegionalConsistency)
 					{
-						bool missing = false;
-						try
+						if (_settings.CurrentValue.RegionalConsistencyCheckNamespaces.Contains(ns.ToString()))
 						{
-							missing = await VerifyRegionalConsistencyAsync(ns, bucket, refId);
+							bool missing = false;
+							try
+							{
+								missing = await VerifyRegionalConsistencyAsync(ns, bucket, refId);
+							}
+							catch (Exception e)
+							{
+								_logger.LogWarning(e, "Unknown exception with message {Message} when checking ref store consistency. Ignoring.", e.Message);
+							}
+
+							if (missing)
+							{
+								Interlocked.Increment(ref countOfRegionalInconsistentRefs);
+							}
 						}
-						catch (Exception e)
+					}
+					if (_settings.CurrentValue.CheckRefStorePartialObjects)
+					{
+						// check if any ref can be found by listing but not when explicitly looked up, indicates it doesn't contain all the expected fields and thus is a partial record that should be removed to fix bugs with how TTL was set when is_finalized was updated
+						try
 						{
-							_logger.LogWarning(e, "Unknown exception with message {Message} when checking ref store consistency. Ignoring.", e.Message);
+							await _referencesStore.GetAsync(ns, bucket, refId, IReferencesStore.FieldFlags.None, IReferencesStore.OperationFlags.BypassCache, token);
 						}
-
-						if (missing)
+						catch (RefNotFoundException)
 						{
-							Interlocked.Increment(ref countOfRegionalInconsistentRefs);
+							_logger.LogWarning("Partial record {Namespace} {Bucket} {Name} found. Deleting.", ns, bucket, refId);
+							await _referencesStore.DeleteAsync(ns, bucket, refId, token);
 						}
 					}
-				}
-				if (_settings.CurrentValue.CheckRefStorePartialObjects)
-				{
-					// check if any ref can be found by listing but not when explicitly looked up, indicates it doesn't contain all the expected fields and thus is a partial record that should be removed to fix bugs with how TTL was set when is_finalized was updated
-					try
-					{
-						await _referencesStore.GetAsync(ns, bucket, refId, IReferencesStore.FieldFlags.None, IReferencesStore.OperationFlags.BypassCache, cancellationToken);
-					}
-					catch (RefNotFoundException)
-					{
-						_logger.LogWarning("Partial record {Namespace} {Bucket} {Name} found. Deleting.", ns, bucket, refId);
-						await _referencesStore.DeleteAsync(ns, bucket, refId, cancellationToken);
-					}
-				}
-				Interlocked.Increment(ref countOfRefsChecked);
-			}
+					Interlocked.Increment(ref countOfRefsChecked);
+				});
 
 			if (_settings.CurrentValue.CheckRefStoreLastAccessTimeConsistency)
 			{
