Cloud DDC - Fixed speculative issue with blob consistency check still removing blobs that were unknown if it had never seen it in the region were the check was running. Now checks if it has been seen in any region.

diff --git a/Engine/Source/Programs/UnrealCloudDDC/Jupiter/Implementation/Blob/BlobIndexConsistencyCheckService.cs b/Engine/Source/Programs/UnrealCloudDDC/Jupiter/Implementation/Blob/BlobIndexConsistencyCheckService.cs
--- a/Engine/Source/Programs/UnrealCloudDDC/Jupiter/Implementation/Blob/BlobIndexConsistencyCheckService.cs
+++ b/Engine/Source/Programs/UnrealCloudDDC/Jupiter/Implementation/Blob/BlobIndexConsistencyCheckService.cs
@@ -100,14 +100,8 @@ namespace Jupiter.Implementation
 						bool exists = true;
 						try
 						{
-							BlobMetadata blobMetadata = await _blobService.GetObjectMetadataAsync(ns, blobIdentifier, cancellationToken);
-
-							DateTime cutoff = DateTime.Now.AddMinutes(-60);
-							if (blobMetadata.CreationTime > cutoff)
-							{
-								// skip considering any blobs to make sure they have had time to have their blobs uploaded ands refs to it created
-								return;
-							}
+							List<string> regions = await _blobIndex.GetBlobRegionsAsync(ns, blobIdentifier, cancellationToken);
+							exists = regions.Count != 0;
 						}
 						catch (BlobNotFoundException)
 						{
