Cloud DDC - Fixed build error caused by empty string used when not listening for ScyllaDB traces.

diff --git a/Engine/Source/Programs/UnrealCloudDDC/Jupiter.Common/BaseStartup.cs b/Engine/Source/Programs/UnrealCloudDDC/Jupiter.Common/BaseStartup.cs
--- a/Engine/Source/Programs/UnrealCloudDDC/Jupiter.Common/BaseStartup.cs
+++ b/Engine/Source/Programs/UnrealCloudDDC/Jupiter.Common/BaseStartup.cs
@@ -368,16 +368,25 @@ namespace Jupiter
 				{
 					builder.AddConsoleExporter();
 				}
-				builder.AddSource("UnrealCloudDDC", traceScyllaRequests ? "ScyllaDB" : "");
+				builder.AddSource("UnrealCloudDDC");
+
+				if (traceScyllaRequests)
+				{
+					builder.AddSource("ScyllaDB");
+				}
 			}).WithMetrics(builder =>
 			{
 				builder
-					.AddMeter("UnrealCloudDDC", traceScyllaRequests ? "ScyllaDB" : "")
+					.AddMeter("UnrealCloudDDC")
 					.AddOtlpExporter()
 					.AddAspNetCoreInstrumentation()
 					.AddHttpClientInstrumentation()
 					.AddRuntimeInstrumentation();
 
+				if (traceScyllaRequests)
+				{
+					builder.AddMeter("ScyllaDB");
+				}
 				if (useConsoleExporter)
 				{
 					builder.AddConsoleExporter();
