Cloud DDC - Refactored blob exists checks to use parallel for each instead of manually spawning tasks for each blob, gives more control over how wide it goes at the cost of potentially running a bit less parallel by default. Hopefully resolves issues with broken pipe responses from nginx on occasion.

diff --git a/Engine/Source/Programs/UnrealCloudDDC/Jupiter/Controllers/CompressedBlobController.cs b/Engine/Source/Programs/UnrealCloudDDC/Jupiter/Controllers/CompressedBlobController.cs
--- a/Engine/Source/Programs/UnrealCloudDDC/Jupiter/Controllers/CompressedBlobController.cs
+++ b/Engine/Source/Programs/UnrealCloudDDC/Jupiter/Controllers/CompressedBlobController.cs
@@ -156,9 +156,9 @@ namespace Jupiter.Controllers
 			ConcurrentBag<ContentId> partialContentIds = new ConcurrentBag<ContentId>();
 			ConcurrentBag<ContentId> invalidContentIds = new ConcurrentBag<ContentId>();
 
-			IEnumerable<Task> tasks = id.Select(async blob =>
+			await Parallel.ForEachAsync(id, async (blob, token) =>
 			{
-				BlobId[]? chunks = await _contentIdStore.ResolveAsync(ns, blob, mustBeContentId: false);
+				BlobId[]? chunks = await _contentIdStore.ResolveAsync(ns, blob, mustBeContentId: false, cancellationToken: token);
 
 				if (chunks == null)
 				{
@@ -168,14 +168,13 @@ namespace Jupiter.Controllers
 
 				foreach (BlobId chunk in chunks)
 				{
-					if (!await _storage.ExistsAsync(ns, chunk))
+					if (!await _storage.ExistsAsync(ns, chunk, cancellationToken: token))
 					{
 						partialContentIds.Add(blob);
 						break;
 					}
 				}
 			});
-			await Task.WhenAll(tasks);
 
 			List<ContentId> needs = new List<ContentId>(invalidContentIds);
 			needs.AddRange(partialContentIds);
@@ -198,9 +197,9 @@ namespace Jupiter.Controllers
 			ConcurrentBag<ContentId> partialContentIds = new ConcurrentBag<ContentId>();
 			ConcurrentBag<ContentId> invalidContentIds = new ConcurrentBag<ContentId>();
 
-			IEnumerable<Task> tasks = bodyIds.Select(async blob =>
+			await Parallel.ForEachAsync(bodyIds, async (blob, token) =>
 			{
-				BlobId[]? chunks = await _contentIdStore.ResolveAsync(ns, blob, mustBeContentId: false);
+				BlobId[]? chunks = await _contentIdStore.ResolveAsync(ns, blob, mustBeContentId: false, cancellationToken: token);
 
 				if (chunks == null)
 				{
@@ -210,14 +209,13 @@ namespace Jupiter.Controllers
 
 				foreach (BlobId chunk in chunks)
 				{
-					if (!await _storage.ExistsAsync(ns, chunk))
+					if (!await _storage.ExistsAsync(ns, chunk, cancellationToken: token))
 					{
 						partialContentIds.Add(blob);
 						break;
 					}
 				}
 			});
-			await Task.WhenAll(tasks);
 
 			List<ContentId> needs = new List<ContentId>(invalidContentIds);
 			needs.AddRange(partialContentIds);
