Cloud DDC - Handle task cancelletion during batch endpoint execution. Its fully expected that a batch request may take some time and if the caller aborts then we should just stop what we are doing but its not an error.

diff --git a/Engine/Source/Programs/UnrealCloudDDC/Jupiter/Controllers/ReferencesController.cs b/Engine/Source/Programs/UnrealCloudDDC/Jupiter/Controllers/ReferencesController.cs
--- a/Engine/Source/Programs/UnrealCloudDDC/Jupiter/Controllers/ReferencesController.cs
+++ b/Engine/Source/Programs/UnrealCloudDDC/Jupiter/Controllers/ReferencesController.cs
@@ -1236,6 +1236,10 @@ namespace Jupiter.Controllers
 
 					return (cb, HttpStatusCode.OK);
 				}
+				catch (Exception ex) when (ex is TaskCanceledException)
+				{
+					return ToErrorResult(ex, HttpStatusCode.RequestTimeout);
+				}
 				catch (Exception ex) when (ex is RefNotFoundException or PartialReferenceResolveException or ReferenceIsMissingBlobsException)
 				{
 					return ToErrorResult(ex, HttpStatusCode.NotFound);
@@ -1277,6 +1281,10 @@ namespace Jupiter.Controllers
 
 					return (CbObject.Build(writer => writer.WriteBool("exists", true)), HttpStatusCode.OK);
 				}
+				catch (Exception ex) when (ex is TaskCanceledException)
+				{
+					return ToErrorResult(ex, HttpStatusCode.RequestTimeout);
+				}
 				catch (Exception ex) when (ex is RefNotFoundException or PartialReferenceResolveException or ReferenceIsMissingBlobsException)
 				{
 					return (CbObject.Build(writer => writer.WriteBool("exists", false)), HttpStatusCode.NotFound);
@@ -1313,6 +1321,10 @@ namespace Jupiter.Controllers
 
 					return (CbSerializer.Serialize(new PutObjectResponse(missingHashes.ToArray())), HttpStatusCode.OK);
 				}
+				catch (Exception ex) when (ex is TaskCanceledException)
+				{
+					return ToErrorResult(ex, HttpStatusCode.RequestTimeout);
+				}
 				catch (Exception e)
 				{
 					return ToErrorResult(e);
