Cloud DDC - Limit concurrency on blob exists checks to see if this impacts the connection resets we sometimes see within nginx for this endpoint

diff --git a/Engine/Source/Programs/UnrealCloudDDC/Jupiter/Controllers/CompressedBlobController.cs b/Engine/Source/Programs/UnrealCloudDDC/Jupiter/Controllers/CompressedBlobController.cs
--- a/Engine/Source/Programs/UnrealCloudDDC/Jupiter/Controllers/CompressedBlobController.cs
+++ b/Engine/Source/Programs/UnrealCloudDDC/Jupiter/Controllers/CompressedBlobController.cs
@@ -156,7 +156,8 @@ namespace Jupiter.Controllers
 			ConcurrentBag<ContentId> partialContentIds = new ConcurrentBag<ContentId>();
 			ConcurrentBag<ContentId> invalidContentIds = new ConcurrentBag<ContentId>();
 
-			await Parallel.ForEachAsync(id, async (blob, token) =>
+			// we limit the concurrency to see if this helps reduce broken pipe / connection resets from nginx when this method is used
+			await Parallel.ForEachAsync(id, new ParallelOptions {MaxDegreeOfParallelism = 4},async (blob, token) =>
 			{
 				BlobId[]? chunks = await _contentIdStore.ResolveAsync(ns, blob, mustBeContentId: false, cancellationToken: token);
 
