Cloud DDC - Increased max total request size and header sizes allowed in kestrel, trying to see if this resolves occasional connection request for exists checks that send fairly large headers (a lot of query parameters)

diff --git a/Engine/Source/Programs/UnrealCloudDDC/Jupiter.Common/BaseProgram.cs b/Engine/Source/Programs/UnrealCloudDDC/Jupiter.Common/BaseProgram.cs
--- a/Engine/Source/Programs/UnrealCloudDDC/Jupiter.Common/BaseProgram.cs
+++ b/Engine/Source/Programs/UnrealCloudDDC/Jupiter.Common/BaseProgram.cs
@@ -128,6 +128,11 @@ namespace Jupiter
 						// Decrease min request body rate to be more forgiving for clients that are going wide and not sending as much data per connection
 						options.Limits.MinRequestBodyDataRate = new MinDataRate(10, TimeSpan.FromSeconds(60));
 
+						// increase max request size of 1 MB to 8 MB
+						options.Limits.MaxRequestBufferSize = 8 * 1024 * 1024;
+						// increase header size to 1 MB
+						options.Limits.MaxRequestHeadersTotalSize = 1024 * 1024;
+
 						string socketsRoot = settings.DomainSocketsRoot;
 
 						if (settings.UseDomainSockets)
