Cloud DDC - Prepared release 1.3.0

diff --git a/Engine/Source/Programs/UnrealCloudDDC/Helm/UnrealCloudDDC/version.yaml b/Engine/Source/Programs/UnrealCloudDDC/Helm/UnrealCloudDDC/version.yaml
--- a/Engine/Source/Programs/UnrealCloudDDC/Helm/UnrealCloudDDC/version.yaml
+++ b/Engine/Source/Programs/UnrealCloudDDC/Helm/UnrealCloudDDC/version.yaml
@@ -1 +1 @@
-version: 1.2.0
\ No newline at end of file
+version: 1.3.0
\ No newline at end of file
diff --git a/Engine/Source/Programs/UnrealCloudDDC/changelog.md b/Engine/Source/Programs/UnrealCloudDDC/changelog.md
--- a/Engine/Source/Programs/UnrealCloudDDC/changelog.md
+++ b/Engine/Source/Programs/UnrealCloudDDC/changelog.md
@@ -1,5 +1,68 @@
 # Unreleased
 
+# 1.3.0
+Noteworthy and potentially breaking changes:
+* Failing upload of compressed buffers if more data then intended is sent.
+If a valid compressed buffer was sent and then extra data appended to it at the end, Cloud DDC would accept the full data and not just the initial buffer sent and store that. This is not a frequently occurring issue but rather something we ran into during development of the builds api.
+These kinds of uploads will now only store and serve the actual compressed buffer and discards the extra data.
+
+* Extra services in helm chart now requires a `enabled` flag
+This breaking change was introduced to make it easier to manage multiple helm charts with different extra services per file so that we could move most of their definitions to a common base file while then enabling them per region we deploy them in.
+
+New features:
+* [Experimental] This releases contains a initial version of the build api which we use for build distribution of staged and packaged builds as well as cooked output. This is not yet supported or documented and may still change.
+* [Experimental] As part of build out the builds api we have a experimental dashboard for Cloud DDC with some limited functionality for admins. The source for this is not yet provided as we have yet to decide if we want to keep this around but you will see some references to it especially in the helm charts.
+* S3 operations now attempt to use multipart for both GET and PUT, GET only works for blobs that were uploaded with multipart PUT. Mostly useful for large files in the builds api.
+* Cloud DDC can now serve a yaml file that describes which endpoints are deployed and controls based on your public ip which instance to talk to. We use this to control which local cache servers and regional replicas to use. Only used by tools in the builds api.
+* Authentication configuration can now be served up in a public unauthenticated endpoint, if you use this make sure that it doesn't include any sensitive information. We encrypt this data during transmission to obfuscate it but it should not be considered safe. This is useful to allow tools to authenticate gainst Cloud DDC without requiring the `oidc-configuration.json` file which is typically found in the workspace. Mostly used by the builds api and its related tools.
+* Added option to prevent overwrites of refs per namespace, is opt in per namespace by setting AllowOverwritesOfRefs (defaults to true, old behavior). Can be overridden per request by specifying the X-Jupiter-Allow-Overwrite header (with value true). We are using this to help find and fix inconsistencies within DDC types. Will likely be made the default behavior in the future when we have addressed all the issues encountered.
+* Introduced a new way to configure auth which is currently only used by the builds api (but will replace the method used for DDC as well in the future). This allows us to control access not just per namespace but also by applying regular expressions on the bucket names.
+
+Minor improvements:
+* Optimized reference resolving especially for very large refs (shader databases and oplogs)
+* Avoid buffering payloads during a cache miss if the stream is seekable (so we can upload it to all layers that are missing it)
+* Disabled scylla tracing by default as its very large volume , can be re-enabled using the OTEL_TRACE_SCYLLA_REQUESTS environment variable.
+* Added UE-isBuildMachine header to list of headers we export to tracing
+* Trim away extra whitespace in claims when checking if what we get from IdP matches what we expect. Makes it less annoying to configure the IdP as it can be quite hard to spot a extra whitespace in your configuration for some of them.
+* Reconfigured forwarding headers to allow for more then one forwarding hop and removed assumption on proxy being on localhost. Makes the source ip of the connection correct in our audit logs.
+* Added ability to GC the filesystem cache per namespace and to set a max size per namespace. This allows us to make sure that one busy namespace doesn't out compete another important but less frequently accessed namespace.
+* Added namespace option to bypass the filesystem cache on write, useful for oplogs that are write heavy operations and were they are not always read. Gives better utilization of the filesystem cache.
+* Returning a X-Jupiter-IoHash header when retrieving a compressed blob. This is the hash of the compressed content so it can be used by clients to verify that the content received matches with what is expected without decompressing it first.
+* Changed the auth endpoint to use the new method for controlling auth to run its checks which also includes debug logging that explains the results.
+* Deleting namespaces and buckets that are found to exist but not be used.
+* Allow all origins to call into Jupiter in CORS - enabling people to use their own web apps to call the api. They will need to send auth tokens anyway so no need to prevent CORS.
+
+Bug fixes:
+
+* Various improvements to tracing and logging
+* Improvements to temp file management including bug fixes and better debugging information.
+* Improved the reliability of the replication tasks - especially during periods of very high load
+* Fixed issue were status controllers didn't support requesting compact binary responses on its endpoints.
+* Removed internal filter from admin controller, meaning that these apis can be accessed on the public endpoints (but requires a admin token). They are not super useful if they can not be easily accessed as they are usually used by humans and not by the system itself (which is were the internal filter is mostly useful).
+* Fixed temp file leaks when using replication snapshots
+* Include information on when the ref was created in the replication consistency check message so that we can determine if the issues we are seeing are due to new refs or old ones. (new ones would be expected to take a bit of time to replicate)
+* Upgraded Serilog.Enrichers.ClientInfo to fix vulnerability warning
+* Suppressing exceptions raised during ref store consistency check - we do not particularly care if a ref is in a partial state for this operation.
+* Handle bucket already exists errors from S3 while creating the bucket (this indicates that we had two attempts at creating the bucket at the same time, so the second one fails, as we just care about the bucket existing we can safely ignore this error)
+* Fixed issue were mongo content id comparisons were failing because .Equals is not supported in mongo queries for strings.
+* Upgraded Blake3 version to fix inconsistency with EpicGames.Iohash
+* Upgraded Microsoft.Extensions.Caching.Memory to 8.0.1 to fix build issues
+* Fixed issue where the peer store would call into itself to see if the blob was present which makes no sense.
+* Fix issue were the blob cleanup would fail if no blobs had been written to the filesystem yet.
+* Always require requests to opt in to getting a redirect uri response even if its allowed to be returned for the namespace policy. This makes sure that only clients that actually supports this gets a redirect.
+* Fixed issue were large blobs that could not be inlined was not added to the blob replication log.
+* Fixed issue were settings applied to the Jupiter section were not being passed if you were also using nginx domain sockets. This especially impacted the 'PendingConnectionMax' setting which you bump when you also bump nginx values.
+* Fixed issue were the speculative blob replication wouldn't work when combined with on demand replication (as it noticed that blobs were always available in other regions and thus already existed)
+* Enabled minio browser on port 9001 when using the tests compose, useful to be able to verify the contents after tests have run.
+
+
+Helm chart changes:
+* Fixed issue were securityContexts for the pods was unable to be specified.
+* Fixed issue with helm state always diffing causing ArgoCD to apply helm over and over again.
+* Removed pod ip setting as its not always supported in k8s and we do not need it for anything right now.
+* Upgraded dotnet-monitor to 8.0.5
+
+
 # 1.2.0
 * `InlineMaxBlobSize` option has moved from `Scylla` to `UnrealCloudDDC` as this can now be used to generally control if blob are inlined into the ref store (minimal practical difference as only Scylla supports inlining blobs). Resolves issue with `EnablePutRefBodyIntoBlobStore` option which will prevent these inlined blobs from also being added to the blob store.
 * Added experimental option that is disabled by default to allow enumeration of buckets (enabled per namespace). Could potentially be used by oplogs.
diff --git a/Engine/Source/Programs/UnrealCloudDDC/version.yaml b/Engine/Source/Programs/UnrealCloudDDC/version.yaml
--- a/Engine/Source/Programs/UnrealCloudDDC/version.yaml
+++ b/Engine/Source/Programs/UnrealCloudDDC/version.yaml
@@ -1 +1 @@
-version: 1.2.0
\ No newline at end of file
+version: 1.3.0
\ No newline at end of file
