Cloud DDC - Fixed copy paste error causing file attributes to not be serialized (as field name was incorrect)

diff --git a/Engine/Source/Programs/UnrealCloudDDC/Jupiter/Controllers/BuildsController.cs b/Engine/Source/Programs/UnrealCloudDDC/Jupiter/Controllers/BuildsController.cs
--- a/Engine/Source/Programs/UnrealCloudDDC/Jupiter/Controllers/BuildsController.cs
+++ b/Engine/Source/Programs/UnrealCloudDDC/Jupiter/Controllers/BuildsController.cs
@@ -1831,7 +1831,7 @@ namespace Jupiter.Controllers
 
 		[CbField("rawSizes")]
 		public ulong[] RawSizes { get;set; } = Array.Empty<ulong>();
-		[CbField("paths")]
+		[CbField("attributes")]
 		public uint[] Attributes { get; set; } = Array.Empty<uint>();
 		[CbField("modes")]
 		public uint[] Modes { get; set; } = Array.Empty<uint>();
