Cloud DDC - Added logging when S3 io exceptions happen to make sure we are capturing them correctly in preperation for retrying

diff --git a/Engine/Source/Programs/UnrealCloudDDC/Jupiter/Implementation/Blob/AmazonS3Store.cs b/Engine/Source/Programs/UnrealCloudDDC/Jupiter/Implementation/Blob/AmazonS3Store.cs
--- a/Engine/Source/Programs/UnrealCloudDDC/Jupiter/Implementation/Blob/AmazonS3Store.cs
+++ b/Engine/Source/Programs/UnrealCloudDDC/Jupiter/Implementation/Blob/AmazonS3Store.cs
@@ -7,6 +7,7 @@ using System.Collections.Generic;
 using System.IO;
 using System.Linq;
 using System.Net;
+using System.Net.Http;
 using System.Runtime.CompilerServices;
 using System.Threading;
 using System.Threading.Tasks;
@@ -467,6 +468,20 @@ namespace Jupiter.Implementation
 
 					throw;
 				}
+				catch (HttpIOException)
+				{
+					// just log these to make sure we are correctly capturing when these exceptions happen within the s3 sdk
+					// TODO: Add retry
+					_logger.LogWarning("Http IO exception when writing multipart blob to S3");
+					throw;
+				}
+				catch (IOException)
+				{
+					// just log these to make sure we are correctly capturing when these exceptions happen within the s3 sdk
+					// TODO: Add retry
+					_logger.LogWarning("IO exception when writing multipart blob to S3");
+					throw;
+				}
 			}
 			finally
 			{
