Cloud DDC - Fallback to non-consistent reads if we can not meet quorum requirements to limit impact when a single node is unable to be reached

diff --git a/Engine/Source/Programs/UnrealCloudDDC/Jupiter/Implementation/ContentId/ScyllaContentIdStore.cs b/Engine/Source/Programs/UnrealCloudDDC/Jupiter/Implementation/ContentId/ScyllaContentIdStore.cs
--- a/Engine/Source/Programs/UnrealCloudDDC/Jupiter/Implementation/ContentId/ScyllaContentIdStore.cs
+++ b/Engine/Source/Programs/UnrealCloudDDC/Jupiter/Implementation/ContentId/ScyllaContentIdStore.cs
@@ -55,54 +55,71 @@ namespace Jupiter.Implementation
 			{
 				using TelemetrySpan contentIdFetchScope = _tracer.BuildScyllaSpan("ScyllaContentIdStore.FetchContentId").SetAttribute("resource.name", contentId.ToString());
 
-				string executionProfile = _settings.CurrentValue.UseConsistentContentIdResolve ? "local_consistent" : "default";
-				if (_scyllaSessionManager.IsScylla)
+				const string ConsistentProfile = "local_consistent";
+				const string DefaultProfile = "default";
+				string executionProfile = _settings.CurrentValue.UseConsistentContentIdResolve ? ConsistentProfile : DefaultProfile;
+				const int MaxRetries = 3;
+				for (int i = 0; i < MaxRetries; i++)
 				{
-					// lower content_weight means its a better candidate to resolve to
-					foreach (ScyllaContentId? resolvedContentId in await _mapper.FetchAsync<ScyllaContentId>(Cql.New("WHERE content_id = ? ORDER BY content_weight ASC", new ScyllaBlobIdentifier(contentId)).WithExecutionProfile(executionProfile)))
+					try
 					{
-						if (resolvedContentId == null)
+						if (_scyllaSessionManager.IsScylla)
 						{
-							throw new InvalidContentIdException(contentId);
-						}
-
-						BlobId[] blobs = resolvedContentId.Chunks.Select(b => b.AsBlobIdentifier()).ToArray();
-
-						{
-							using TelemetrySpan _ = _tracer.StartActiveSpan("ScyllaContentIdStore.FindMissingBlobs").SetAttribute("operation.name", "ScyllaContentIdStore.FindMissingBlobs");
-
-							BlobId[] missingBlobs = await _blobStore.FilterOutKnownBlobsAsync(ns, blobs, cancellationToken);
-							if (missingBlobs.Length == 0)
+							// lower content_weight means its a better candidate to resolve to
+							foreach (ScyllaContentId? resolvedContentId in await _mapper.FetchAsync<ScyllaContentId>(Cql.New("WHERE content_id = ? ORDER BY content_weight ASC", new ScyllaBlobIdentifier(contentId)).WithExecutionProfile(executionProfile)))
 							{
-								return blobs;
+								if (resolvedContentId == null)
+								{
+									throw new InvalidContentIdException(contentId);
+								}
+
+								BlobId[] blobs = resolvedContentId.Chunks.Select(b => b.AsBlobIdentifier()).ToArray();
+
+								{
+									using TelemetrySpan _ = _tracer.StartActiveSpan("ScyllaContentIdStore.FindMissingBlobs").SetAttribute("operation.name", "ScyllaContentIdStore.FindMissingBlobs");
+
+									BlobId[] missingBlobs = await _blobStore.FilterOutKnownBlobsAsync(ns, blobs, cancellationToken);
+									if (missingBlobs.Length == 0)
+									{
+										return blobs;
+									}
+								}
+								// blobs are missing continue testing with the next content id in the weighted list as that might exist
 							}
 						}
-						// blobs are missing continue testing with the next content id in the weighted list as that might exist
-					}
-				}
-				else
-				{
-					// lower content_weight means its a better candidate to resolve to
-					foreach (CassandraContentId? resolvedContentId in await _mapper.FetchAsync<CassandraContentId>(Cql.New("WHERE content_id = ? ORDER BY content_weight ASC", contentId.HashData).WithExecutionProfile(executionProfile)))
-					{
-						if (resolvedContentId == null)
+						else
 						{
-							throw new InvalidContentIdException(contentId);
-						}
-
-						BlobId[] blobs = resolvedContentId.Chunks.Select(b => new BlobId(b!)).ToArray();
-
-						{
-							using TelemetrySpan _ = _tracer.StartActiveSpan("ScyllaContentIdStore.FindMissingBlobs").SetAttribute("operation.name", "ScyllaContentIdStore.FindMissingBlobs");
-
-							BlobId[] missingBlobs = await _blobStore.FilterOutKnownBlobsAsync(ns, blobs, cancellationToken);
-							if (missingBlobs.Length == 0)
+							// lower content_weight means its a better candidate to resolve to
+							foreach (CassandraContentId? resolvedContentId in await _mapper.FetchAsync<CassandraContentId>(Cql.New("WHERE content_id = ? ORDER BY content_weight ASC", contentId.HashData).WithExecutionProfile(executionProfile)))
 							{
-								return blobs;
+								if (resolvedContentId == null)
+								{
+									throw new InvalidContentIdException(contentId);
+								}
+
+								BlobId[] blobs = resolvedContentId.Chunks.Select(b => new BlobId(b!)).ToArray();
+
+								{
+									using TelemetrySpan _ = _tracer.StartActiveSpan("ScyllaContentIdStore.FindMissingBlobs").SetAttribute("operation.name", "ScyllaContentIdStore.FindMissingBlobs");
+
+									BlobId[] missingBlobs = await _blobStore.FilterOutKnownBlobsAsync(ns, blobs, cancellationToken);
+									if (missingBlobs.Length == 0)
+									{
+										return blobs;
+									}
+								}
+								// blobs are missing continue testing with the next content id in the weighted list as that might exist
 							}
 						}
-						// blobs are missing continue testing with the next content id in the weighted list as that might exist
 					}
+					catch (UnavailableException)
+					{
+						// not enough replicas available fall back to non-consistent reads to limit impact
+						executionProfile = DefaultProfile;
+						continue;
+					}
+
+					break;
 				}
 			}
 
