Cloud DDC - Added search option to include non-finalized builds. These should not be attempted to be downloaded but it can be useful to show that they exist for admins.

diff --git a/Engine/Source/Programs/UnrealCloudDDC/Jupiter/Controllers/BuildsController.cs b/Engine/Source/Programs/UnrealCloudDDC/Jupiter/Controllers/BuildsController.cs
--- a/Engine/Source/Programs/UnrealCloudDDC/Jupiter/Controllers/BuildsController.cs
+++ b/Engine/Source/Programs/UnrealCloudDDC/Jupiter/Controllers/BuildsController.cs
@@ -346,6 +346,7 @@ namespace Jupiter.Controllers
 			DateTime startTime = DateTime.Now;
 			TimeSpan maxTimeToSpend = TimeSpan.FromMilliseconds(searchRequest.Options.MaxTimeMs);
 			bool includeTTL = searchRequest.Options.IncludeTTL;
+			bool includePartial = searchRequest.Options.IncludePartial;
 			bool partialResult = false;
 			TimeSpan timeSpent = TimeSpan.Zero;
 
@@ -353,9 +354,9 @@ namespace Jupiter.Controllers
 
 			await foreach (BuildMetadata build in _buildStore.ListBuildsAsync(ns, bucket, includeTTL))
 			{
-				if (!build.IsFinalized)
+				if (!build.IsFinalized && !includePartial)
 				{
-					// we ignore non-finalized builds
+					// we ignore non-finalized builds unless explicitly asked to include them
 					continue;
 				}
 				count++;
@@ -470,6 +471,7 @@ namespace Jupiter.Controllers
 			DateTime startTime = DateTime.Now;
 			TimeSpan maxTimeToSpend = TimeSpan.FromMilliseconds(searchRequest.Options.MaxTimeMs);
 			bool includeTTL = searchRequest.Options.IncludeTTL;
+			bool includePartial = searchRequest.Options.IncludePartial;
 			bool partialResult = false;
 			TimeSpan timeSpent = TimeSpan.Zero;
 			bool timedOut = false;
@@ -485,9 +487,9 @@ namespace Jupiter.Controllers
 				
 					await foreach (BuildMetadata build in _buildStore.ListBuildsAsync(ns, bucket, includeTTL).WithCancellation(token))
 					{
-						if (!build.IsFinalized)
+						if (!build.IsFinalized && !includePartial)
 						{
-							// we ignore non-finalized builds
+							// we ignore non-finalized builds unless explicitly asked to include them
 							continue;
 						}
 
@@ -1981,6 +1983,9 @@ namespace Jupiter.Controllers
 		public uint? TTL { get; set; }
 		[CbField("metadata")]
 		public CbObject Metadata { get; set; } = null!;
+
+		[CbField("isFinalized")]
+		public bool IsFinalized { get; set; } = true;
 	}
 	
 	public class SearchResult
@@ -2003,7 +2008,7 @@ namespace Jupiter.Controllers
 		public SearchResult(List<BuildMetadata> foundResults, bool partialResult)
 		{
 			PartialResult = partialResult;
-			Results = foundResults.Select(b => new BuildSearchResult {BuildId = b.BuildId, Metadata = b.Metadata, TTL = b.Ttl}).ToList();
+			Results = foundResults.Select(b => new BuildSearchResult {BuildId = b.BuildId, Metadata = b.Metadata, TTL = b.Ttl, IsFinalized = b.IsFinalized}).ToList();
 		}
 	}
 
@@ -2023,6 +2028,11 @@ namespace Jupiter.Controllers
 
 		[CbField("includeTTL")]
 		public bool IncludeTTL { get; set; } = false;
+
+		
+		// a custom option to include builds that are not finalized, e.g. they are only partially uploaded
+		[CbField("includePartial")]
+		public bool IncludePartial { get; set; } = false;
 	}
 
 	public class SearchRequest
