Cloud DDC - Retry failed mutlipart uploads to S3

diff --git a/Engine/Source/Programs/UnrealCloudDDC/Jupiter/Implementation/Blob/AmazonS3Store.cs b/Engine/Source/Programs/UnrealCloudDDC/Jupiter/Implementation/Blob/AmazonS3Store.cs
--- a/Engine/Source/Programs/UnrealCloudDDC/Jupiter/Implementation/Blob/AmazonS3Store.cs
+++ b/Engine/Source/Programs/UnrealCloudDDC/Jupiter/Implementation/Blob/AmazonS3Store.cs
@@ -441,46 +441,63 @@ namespace Jupiter.Implementation
 
 					filePath = payload.TempFile.FullName;
 				}
-				using TransferUtility utility = new TransferUtility(_amazonS3);
-				try
+
+				const int MaxRetryAttempts = 3;
+				for (int i = 0; i < MaxRetryAttempts; i++)
 				{
-					if (filePath != null)
+					using TransferUtility utility = new TransferUtility(_amazonS3);
+					try
 					{
-						await utility.UploadAsync(filePath, _bucketName, path, cancellationToken);
+						if (filePath != null)
+						{
+							await utility.UploadAsync(filePath, _bucketName, path, cancellationToken);
+						}
+						else
+						{
+							await utility.UploadAsync(stream, _bucketName, path, cancellationToken);
+						}
+
+						break;
 					}
-					else
+					catch (AmazonS3Exception e)
 					{
-						await utility.UploadAsync(stream, _bucketName, path, cancellationToken);
+						// if the same object is added twice S3 will raise a error, as we are content addressed we can just accept whichever of the objects so we can ignore that error
+						if (e.StatusCode == HttpStatusCode.Conflict)
+						{
+							return;
+						}
+
+						if (e.StatusCode == HttpStatusCode.TooManyRequests)
+						{
+							throw new ResourceHasToManyRequestsException(e);
+						}
+
+						throw;
 					}
-				}
-				catch (AmazonS3Exception e)
-				{
-					// if the same object is added twice S3 will raise a error, as we are content addressed we can just accept whichever of the objects so we can ignore that error
-					if (e.StatusCode == HttpStatusCode.Conflict)
+					catch (HttpIOException e)
 					{
-						return;
-					}
+						if (i < MaxRetryAttempts-1)
+						{
+							// retry
+							continue;
+						}
 
-					if (e.StatusCode == HttpStatusCode.TooManyRequests)
-					{
-						throw new ResourceHasToManyRequestsException(e);
+						// to many retries, give up
+						_logger.LogError(e, "Http IO exception when writing multipart blob to S3");
+						throw;
 					}
+					catch (IOException e)
+					{
+						if (i < MaxRetryAttempts-1)
+						{
+							// retry
+							continue;
+						}
 
-					throw;
-				}
-				catch (HttpIOException)
-				{
-					// just log these to make sure we are correctly capturing when these exceptions happen within the s3 sdk
-					// TODO: Add retry
-					_logger.LogWarning("Http IO exception when writing multipart blob to S3");
-					throw;
-				}
-				catch (IOException)
-				{
-					// just log these to make sure we are correctly capturing when these exceptions happen within the s3 sdk
-					// TODO: Add retry
-					_logger.LogWarning("IO exception when writing multipart blob to S3");
-					throw;
+						// to many retries, give up
+						_logger.LogError(e, "IO exception when writing multipart blob to S3");
+						throw;
+					}
 				}
 			}
 			finally
