Cloud DDC - Copy the stream before uploading it to S3 if its not backed by a file, if the uploads fails this stream will be closed by AWS SDK and thus the retry will fail. Also added some warnings if uploads fail but we retry to better track how often it happens.

diff --git a/Engine/Source/Programs/UnrealCloudDDC/Jupiter/Implementation/Blob/AmazonS3Store.cs b/Engine/Source/Programs/UnrealCloudDDC/Jupiter/Implementation/Blob/AmazonS3Store.cs
--- a/Engine/Source/Programs/UnrealCloudDDC/Jupiter/Implementation/Blob/AmazonS3Store.cs
+++ b/Engine/Source/Programs/UnrealCloudDDC/Jupiter/Implementation/Blob/AmazonS3Store.cs
@@ -454,6 +454,11 @@ namespace Jupiter.Implementation
 						}
 						else
 						{
+							// copy the stream as TransferUtility will dispose this stream if any issue happens making us unable to retry
+							using MemoryStream ms = new MemoryStream();
+							await stream.CopyToAsync(ms, cancellationToken);
+							ms.Position = 0;
+
 							await utility.UploadAsync(stream, _bucketName, path, cancellationToken);
 						}
 
@@ -478,6 +483,7 @@ namespace Jupiter.Implementation
 					{
 						if (i < MaxRetryAttempts-1)
 						{
+							_logger.LogWarning("Http IO exception when writing multipart blob to S3, retrying");
 							// retry
 							continue;
 						}
@@ -490,6 +496,7 @@ namespace Jupiter.Implementation
 					{
 						if (i < MaxRetryAttempts-1)
 						{
+							_logger.LogWarning("Http IO exception when writing multipart blob to S3, retrying");
 							// retry
 							continue;
 						}
