Cloud DDC - Uploading the copied stream to s3 rather then the original stream

diff --git a/Engine/Source/Programs/UnrealCloudDDC/Jupiter/Implementation/Blob/AmazonS3Store.cs b/Engine/Source/Programs/UnrealCloudDDC/Jupiter/Implementation/Blob/AmazonS3Store.cs
--- a/Engine/Source/Programs/UnrealCloudDDC/Jupiter/Implementation/Blob/AmazonS3Store.cs
+++ b/Engine/Source/Programs/UnrealCloudDDC/Jupiter/Implementation/Blob/AmazonS3Store.cs
@@ -459,7 +459,7 @@ namespace Jupiter.Implementation
 							await stream.CopyToAsync(ms, cancellationToken);
 							ms.Position = 0;
 
-							await utility.UploadAsync(stream, _bucketName, path, cancellationToken);
+							await utility.UploadAsync(ms, _bucketName, path, cancellationToken);
 						}
 
 						break;
