Cloud DDC - Copy the stream before uploading to make sure we only copy it once

diff --git a/Engine/Source/Programs/UnrealCloudDDC/Jupiter/Implementation/Blob/AmazonS3Store.cs b/Engine/Source/Programs/UnrealCloudDDC/Jupiter/Implementation/Blob/AmazonS3Store.cs
--- a/Engine/Source/Programs/UnrealCloudDDC/Jupiter/Implementation/Blob/AmazonS3Store.cs
+++ b/Engine/Source/Programs/UnrealCloudDDC/Jupiter/Implementation/Blob/AmazonS3Store.cs
@@ -423,6 +423,7 @@ namespace Jupiter.Implementation
 			FilesystemBufferedPayload? payload = null;
 			try
 			{
+				await using MemoryStream localStream = new MemoryStream();
 				string? filePath = null;
 				if (stream is FileStream fileStream)
 				{
@@ -441,6 +442,12 @@ namespace Jupiter.Implementation
 
 					filePath = payload.TempFile.FullName;
 				}
+				else
+				{
+					// copy the stream as TransferUtility will dispose this stream if any issue happens making us unable to retry
+					await stream.CopyToAsync(localStream, cancellationToken);
+					localStream.Position = 0;
+				}
 
 				const int MaxRetryAttempts = 3;
 				for (int i = 0; i < MaxRetryAttempts; i++)
@@ -454,12 +461,7 @@ namespace Jupiter.Implementation
 						}
 						else
 						{
-							// copy the stream as TransferUtility will dispose this stream if any issue happens making us unable to retry
-							using MemoryStream ms = new MemoryStream();
-							await stream.CopyToAsync(ms, cancellationToken);
-							ms.Position = 0;
-
-							await utility.UploadAsync(ms, _bucketName, path, cancellationToken);
+							await utility.UploadAsync(localStream, _bucketName, path, cancellationToken);
 						}
 
 						break;
