Cloud DDC - Fixed issue with stream not being reset after each upload attempt making every attempt after the first incorrect

diff --git a/Engine/Source/Programs/UnrealCloudDDC/Jupiter/Implementation/Blob/AmazonS3Store.cs b/Engine/Source/Programs/UnrealCloudDDC/Jupiter/Implementation/Blob/AmazonS3Store.cs
--- a/Engine/Source/Programs/UnrealCloudDDC/Jupiter/Implementation/Blob/AmazonS3Store.cs
+++ b/Engine/Source/Programs/UnrealCloudDDC/Jupiter/Implementation/Blob/AmazonS3Store.cs
@@ -446,7 +446,6 @@ namespace Jupiter.Implementation
 				{
 					// copy the stream as TransferUtility will dispose this stream if any issue happens making us unable to retry
 					await stream.CopyToAsync(localStream, cancellationToken);
-					localStream.Position = 0;
 				}
 
 				const int MaxRetryAttempts = 3;
@@ -461,6 +460,8 @@ namespace Jupiter.Implementation
 						}
 						else
 						{
+							// reset the stream position so that we upload the entire stream each attempt
+							localStream.Position = 0;
 							await utility.UploadAsync(localStream, _bucketName, path, cancellationToken);
 						}
 
