Cloud DDC - Catch task cancelled exceptions when uploading blobs to builds and return explicit exit codes when this happens (408 request timedout)

diff --git a/Engine/Source/Programs/UnrealCloudDDC/Jupiter/Controllers/BuildsController.cs b/Engine/Source/Programs/UnrealCloudDDC/Jupiter/Controllers/BuildsController.cs
--- a/Engine/Source/Programs/UnrealCloudDDC/Jupiter/Controllers/BuildsController.cs
+++ b/Engine/Source/Programs/UnrealCloudDDC/Jupiter/Controllers/BuildsController.cs
@@ -1292,6 +1292,10 @@ namespace Jupiter.Controllers
 			{
 				return Problem(e.Message, null, (int)HttpStatusCode.RequestTimeout);
 			}
+			catch (TaskCanceledException)
+			{
+				return StatusCode(StatusCodes.Status408RequestTimeout, "Request cancelled");
+			}
 		}
 
 		[HttpPost("{ns}/{bucket}/{buildId}/parts/{partId}/blobs/{id}/startMultipartUpload")]
