Cloud DDC - Catching more cases of S3 slow down requests

diff --git a/Engine/Source/Programs/UnrealCloudDDC/Jupiter/Implementation/Blob/AmazonS3Store.cs b/Engine/Source/Programs/UnrealCloudDDC/Jupiter/Implementation/Blob/AmazonS3Store.cs
--- a/Engine/Source/Programs/UnrealCloudDDC/Jupiter/Implementation/Blob/AmazonS3Store.cs
+++ b/Engine/Source/Programs/UnrealCloudDDC/Jupiter/Implementation/Blob/AmazonS3Store.cs
@@ -476,6 +476,11 @@ namespace Jupiter.Implementation
 						}
 
 						if (e.StatusCode == HttpStatusCode.TooManyRequests)
+						{	
+							throw new ResourceHasToManyRequestsException(e);
+						}
+
+						if (e is { StatusCode: HttpStatusCode.ServiceUnavailable, ErrorCode: "SlowDown" or "503 SlowDown" })
 						{
 							throw new ResourceHasToManyRequestsException(e);
 						}
