Cloud DDC - Null check before returning buffer even if it should not be null, helps fix issue with exceptions being hidden during out of memory conditions

diff --git a/Engine/Source/Programs/UnrealCloudDDC/Jupiter/Implementation/Blob/AmazonS3Store.cs b/Engine/Source/Programs/UnrealCloudDDC/Jupiter/Implementation/Blob/AmazonS3Store.cs
--- a/Engine/Source/Programs/UnrealCloudDDC/Jupiter/Implementation/Blob/AmazonS3Store.cs
+++ b/Engine/Source/Programs/UnrealCloudDDC/Jupiter/Implementation/Blob/AmazonS3Store.cs
@@ -593,7 +593,12 @@ namespace Jupiter.Implementation
 					{
 						for (int i = 0; i < partsCount; i++)
 						{
-							ArrayPool<byte>.Shared.Return(parts[i]);
+							// ReSharper disable once ConditionIsAlwaysTrueOrFalseAccordingToNullableAPIContract
+							// This should never be null but during exceptions we have seen it be null so to avoid throwing more exceptions that hide the original exceptions we just null check here
+							if (parts[i] != null)
+							{
+								ArrayPool<byte>.Shared.Return(parts[i]);
+							}
 						}
 					}
 				}
