Cloud DDC - If we are out of memory when buffering payloads (as they can be quite large while still being buffered in memory) fallback to the filesystem even for small payloads.

diff --git a/Engine/Source/Programs/UnrealCloudDDC/Jupiter.Common/Implementation/StreamingBodyPayload.cs b/Engine/Source/Programs/UnrealCloudDDC/Jupiter.Common/Implementation/StreamingBodyPayload.cs
--- a/Engine/Source/Programs/UnrealCloudDDC/Jupiter.Common/Implementation/StreamingBodyPayload.cs
+++ b/Engine/Source/Programs/UnrealCloudDDC/Jupiter.Common/Implementation/StreamingBodyPayload.cs
@@ -198,9 +198,16 @@ namespace Jupiter.Common.Implementation
 		public async Task<IBufferedPayload> CreateFromStreamAsync(Stream s, long contentLength, string debugPrefix, CancellationToken cancellationToken)
 		{
 			// blob is small enough to fit into memory we just read it as is
-			if (contentLength < _options.CurrentValue.MemoryBufferSize)
+			try
 			{
-				return await MemoryBufferedPayload.CreateAsync(_tracer, s, cancellationToken);
+				if (contentLength < _options.CurrentValue.MemoryBufferSize)
+				{
+					return await MemoryBufferedPayload.CreateAsync(_tracer, s, cancellationToken);
+				}
+			}
+			catch (OutOfMemoryException)
+			{
+				// if we fail to allocate the memory buffered payload we fallback to backing it on the filesystem
 			}
 
 			return await FilesystemBufferedPayload.CreateAsync(_tracer, s, _options.CurrentValue.FilesystemTempPayloadRoot, debugPrefix, cancellationToken);
