Cloud DDC - Increased number of builds we allowed to be scanned during a search by default as we could end up in situations were we didn't find builds that existed for normal queries in ushell (for branches that were a bit older)

diff --git a/Engine/Source/Programs/UnrealCloudDDC/Jupiter/Controllers/BuildsController.cs b/Engine/Source/Programs/UnrealCloudDDC/Jupiter/Controllers/BuildsController.cs
--- a/Engine/Source/Programs/UnrealCloudDDC/Jupiter/Controllers/BuildsController.cs
+++ b/Engine/Source/Programs/UnrealCloudDDC/Jupiter/Controllers/BuildsController.cs
@@ -2026,7 +2026,7 @@ namespace Jupiter.Controllers
 		public int Skip { get; set; } = 0;
 
 		[CbField("max")]
-		public int Max { get; set; } = 2000;
+		public int Max { get; set; } = 5000;
 
 		[CbField("maxTimeMS")]
 		public uint MaxTimeMs { get; set; } = Debugger.IsAttached ? 0u : 120_000; // default to 2 minutes max per search unless debugging in which case we disable it
