Cloud DDC - Fixed retries of io exceptions (partial transfers) during blob replication

diff --git a/Engine/Source/Programs/UnrealCloudDDC/Jupiter/Implementation/Replication/BlobsReplicator.cs b/Engine/Source/Programs/UnrealCloudDDC/Jupiter/Implementation/Replication/BlobsReplicator.cs
--- a/Engine/Source/Programs/UnrealCloudDDC/Jupiter/Implementation/Replication/BlobsReplicator.cs
+++ b/Engine/Source/Programs/UnrealCloudDDC/Jupiter/Implementation/Replication/BlobsReplicator.cs
@@ -450,25 +450,54 @@ namespace Jupiter.Implementation
 			_logger.LogDebug("Attempting to replicate blob {Blob} in {Namespace}.", blob, ns);
 
 			const int RetryAttempts = 5;
-			HttpResponseMessage? blobResponse = null;
 			Exception? lastException = null;
 			for (int i = 0; i < RetryAttempts; i++)
 			{
-				using HttpRequestMessage blobRequest = await BuildHttpRequestAsync(HttpMethod.Get, new Uri($"api/v1/blobs/{ns}/{blob}", UriKind.Relative));
 				try
 				{
-					blobResponse = await _httpClient.SendAsync(blobRequest, HttpCompletionOption.ResponseHeadersRead, cancellationToken);
+					
+					using HttpRequestMessage blobRequest = await BuildHttpRequestAsync(HttpMethod.Get, new Uri($"api/v1/blobs/{ns}/{blob}", UriKind.Relative));
+					using HttpResponseMessage blobResponse = await _httpClient.SendAsync(blobRequest, HttpCompletionOption.ResponseHeadersRead, cancellationToken);
+
+					if (blobResponse == null)
+					{
+						throw new Exception("Blob response never set", lastException);
+					}
+
+					if (blobResponse.StatusCode == HttpStatusCode.NotFound)
+					{
+						_logger.LogWarning("Failed to replicate {Blob} in {Namespace} due to it not existing.", blob, ns);
+						return false;
+					}
+
+					if (blobResponse.StatusCode != HttpStatusCode.OK)
+					{
+						_logger.LogError("Bad http response when replicating {Blob} in {Namespace} . Status code: {StatusCode}", blob, ns, blobResponse.StatusCode);
+						return false;
+					}
+
+					await using Stream s = await blobResponse.Content.ReadAsStreamAsync(cancellationToken);
+					long? contentLength = blobResponse.Content.Headers.ContentLength;
+
+					if (contentLength == null)
+					{
+						throw new Exception("Expected content-length on blob response");
+					}
+
+					using IBufferedPayload payload = await _bufferedPayloadFactory.CreateFromStreamAsync(s, contentLength.Value, "blob-replicator", cancellationToken);
+					
+					await _blobService.PutObjectAsync(ns, payload, blob, bucketHint: bucketHint, cancellationToken: cancellationToken);
 					break;
 				}
 				catch (HttpRequestException e)
 				{
-					blobResponse = null;
 					// rethrow unknown exceptions
 					if (e.InnerException is not IOException)
 					{
 						throw;
 					}
 
+					// is an io exception that we should retry
 					lastException = e;
 
 					if (e.StatusCode == HttpStatusCode.NotFound)
@@ -479,34 +508,6 @@ namespace Jupiter.Implementation
 				}
 			}
 
-			if (blobResponse == null)
-			{
-				throw new Exception("Blob response never set", lastException);
-			}
-
-			if (blobResponse.StatusCode == HttpStatusCode.NotFound)
-			{
-				_logger.LogWarning("Failed to replicate {Blob} in {Namespace} due to it not existing.", blob, ns);
-				return false;
-			}
-
-			if (blobResponse.StatusCode != HttpStatusCode.OK)
-			{
-				_logger.LogError("Bad http response when replicating {Blob} in {Namespace} . Status code: {StatusCode}", blob, ns, blobResponse.StatusCode);
-				return false;
-			}
-
-			await using Stream s = await blobResponse.Content.ReadAsStreamAsync(cancellationToken);
-			long? contentLength = blobResponse.Content.Headers.ContentLength;
-
-			if (contentLength == null)
-			{
-				throw new Exception("Expected content-length on blob response");
-			}
-
-			using IBufferedPayload payload = await _bufferedPayloadFactory.CreateFromStreamAsync(s, contentLength.Value, "blob-replicator", cancellationToken);
-
-			await _blobService.PutObjectAsync(ns, payload, blob, bucketHint: bucketHint, cancellationToken: cancellationToken);
 
 			return true;
 		}
