Cloud DDC - Don't fail to delete if bucket does not exist. Rethrow exception if retries are exhausted.

diff --git a/Engine/Source/Programs/UnrealCloudDDC/Jupiter/Implementation/Blob/AmazonS3Store.cs b/Engine/Source/Programs/UnrealCloudDDC/Jupiter/Implementation/Blob/AmazonS3Store.cs
--- a/Engine/Source/Programs/UnrealCloudDDC/Jupiter/Implementation/Blob/AmazonS3Store.cs
+++ b/Engine/Source/Programs/UnrealCloudDDC/Jupiter/Implementation/Blob/AmazonS3Store.cs
@@ -811,6 +811,12 @@ namespace Jupiter.Implementation
 
 		public async Task DeleteAsync(string path, CancellationToken cancellationToken)
 		{
+			// If the bucket does not exist, then there is nothing to delete
+			if (!await AmazonS3Util.DoesS3BucketExistV2Async(_amazonS3, _bucketName))
+			{
+				return;
+			}
+			
 			await _amazonS3.DeleteObjectAsync(_bucketName, path, cancellationToken);
 		}
 
diff --git a/Engine/Source/Programs/UnrealCloudDDC/Jupiter/Implementation/Replication/BlobsReplicator.cs b/Engine/Source/Programs/UnrealCloudDDC/Jupiter/Implementation/Replication/BlobsReplicator.cs
--- a/Engine/Source/Programs/UnrealCloudDDC/Jupiter/Implementation/Replication/BlobsReplicator.cs
+++ b/Engine/Source/Programs/UnrealCloudDDC/Jupiter/Implementation/Replication/BlobsReplicator.cs
@@ -486,6 +486,8 @@ namespace Jupiter.Implementation
 			Exception? lastException = null;
 			for (int i = 0; i < RetryAttempts; i++)
 			{
+				lastException = null;
+				
 				try
 				{
 					
@@ -541,6 +543,11 @@ namespace Jupiter.Implementation
 				}
 			}
 
+			if (lastException != null)
+			{
+				throw lastException;
+			}
+
 			return true;
 		}
 
