Cloud DDC - Update Helm chart to allow replicators to configure a SourceSite instead of a ConnectionString

diff --git a/Engine/Source/Programs/UnrealCloudDDC/Helm/UnrealCloudDDC/templates/configmap-worker.yaml b/Engine/Source/Programs/UnrealCloudDDC/Helm/UnrealCloudDDC/templates/configmap-worker.yaml
--- a/Engine/Source/Programs/UnrealCloudDDC/Helm/UnrealCloudDDC/templates/configmap-worker.yaml
+++ b/Engine/Source/Programs/UnrealCloudDDC/Helm/UnrealCloudDDC/templates/configmap-worker.yaml
@@ -65,11 +65,17 @@ Replication:
 {{ range . }}
   - ReplicatorName: {{ required "ReplicatorName must be specified for a replicator" .ReplicatorName }}
     NamespaceToReplicate: {{ required "Namespace must be specified for a replicator" .Namespace }}
-    ConnectionString: {{ required "replication remote connection string has to be specified" .ConnectionString }}
     MaxParallelReplications: {{ default 128  .MaxParallelReplications }} 
     SkipSnapshot: {{ default true  .SkipSnapshot }}
     PageSize: {{ default 1000  .PageSize }}
     Version: {{ default "Refs"  .Version }}
+    {{- if .SourceSite }}
+    SourceSite: {{ .SourceSite }}
+    {{- else if .ConnectionString }}
+    ConnectionString: {{ .ConnectionString }}
+    {{- else }}
+    {{- fail "Replicator must specify a SourceSite or ConnectionString" }}
+    {{- end }}
 {{- end -}}
 {{- end }}
 
diff --git a/Engine/Source/Programs/UnrealCloudDDC/Helm/UnrealCloudDDC/templates/configmap.yaml b/Engine/Source/Programs/UnrealCloudDDC/Helm/UnrealCloudDDC/templates/configmap.yaml
--- a/Engine/Source/Programs/UnrealCloudDDC/Helm/UnrealCloudDDC/templates/configmap.yaml
+++ b/Engine/Source/Programs/UnrealCloudDDC/Helm/UnrealCloudDDC/templates/configmap.yaml
@@ -101,10 +101,16 @@ Replication:
 {{ range . }}
   - ReplicatorName: {{ required "ReplicatorName must be specified for a replicator" .ReplicatorName }}
     NamespaceToReplicate: {{ required "Namespace must be specified for a replicator" .Namespace }}
-    ConnectionString: {{ required "replication remote connection string has to be specified" .ConnectionString }}
     MaxParallelReplications: {{ default 64  .MaxParallelReplications }}
     SkipSnapshot: {{ default true  .SkipSnapshot }}
     PageSize: {{ default 1000  .PageSize }}
+    {{- if .SourceSite }}
+    SourceSite: {{ .SourceSite }}
+    {{- else if .ConnectionString }}
+    ConnectionString: {{ .ConnectionString }}
+    {{- else }}
+    {{- fail "Replicator must specify a SourceSite or ConnectionString" }}
+    {{- end }}
 {{- end -}}
 {{- end }}
 
@@ -215,6 +221,10 @@ S3:
   UseArnRegion: {{ $s3.UseArnRegion }}
   {{- end }}
 
+  {{- if hasKey $s3 "SiteBucketsPerStoragePool" }}
+  SiteBucketsPerStoragePool: {{ $s3.SiteBucketsPerStoragePool }}
+  {{- end }}
+
   {{- if $s3.StoragePoolBucketOverride }}
   StoragePoolBucketOverride: {{ $s3.StoragePoolBucketOverride | toYaml | nindent 4  }}
   {{- end }}
