Cloud DDC - Pass SiteBucketsPerStoragePool Helm configuration as YAML

diff --git a/Engine/Source/Programs/UnrealCloudDDC/Helm/UnrealCloudDDC/templates/configmap.yaml b/Engine/Source/Programs/UnrealCloudDDC/Helm/UnrealCloudDDC/templates/configmap.yaml
--- a/Engine/Source/Programs/UnrealCloudDDC/Helm/UnrealCloudDDC/templates/configmap.yaml
+++ b/Engine/Source/Programs/UnrealCloudDDC/Helm/UnrealCloudDDC/templates/configmap.yaml
@@ -221,8 +221,8 @@ S3:
   UseArnRegion: {{ $s3.UseArnRegion }}
   {{- end }}
 
-  {{- if hasKey $s3 "SiteBucketsPerStoragePool" }}
-  SiteBucketsPerStoragePool: {{ $s3.SiteBucketsPerStoragePool }}
+  {{- if $s3.SiteBucketsPerStoragePool }}
+  SiteBucketsPerStoragePool: {{ $s3.SiteBucketsPerStoragePool | toYaml | nindent 4 }}
   {{- end }}
 
   {{- if $s3.StoragePoolBucketOverride }}
