Cloud DDC - Added option to disable payload signing in S3, potentially saving quite a bit of cpu at the cost of removing verficiation of payload consistency (transport level consistency still applies)

diff --git a/Engine/Source/Programs/UnrealCloudDDC/Jupiter/Implementation/Blob/AmazonS3Store.cs b/Engine/Source/Programs/UnrealCloudDDC/Jupiter/Implementation/Blob/AmazonS3Store.cs
--- a/Engine/Source/Programs/UnrealCloudDDC/Jupiter/Implementation/Blob/AmazonS3Store.cs
+++ b/Engine/Source/Programs/UnrealCloudDDC/Jupiter/Implementation/Blob/AmazonS3Store.cs
@@ -491,7 +491,9 @@ namespace Jupiter.Implementation
 					BucketName = _bucketName,
 					Key = path,
 					InputStream = stream,
-					UseChunkEncoding = _settings.CurrentValue.UseChunkEncoding
+					UseChunkEncoding = _settings.CurrentValue.UseChunkEncoding,
+					DisablePayloadSigning = _settings.CurrentValue.DisablePayloadChecksums,
+					DisableDefaultChecksumValidation = _settings.CurrentValue.DisablePayloadChecksums,
 				};
 
 				try
@@ -554,13 +556,30 @@ namespace Jupiter.Implementation
 					{
 						if (filePath != null)
 						{
-							await utility.UploadAsync(filePath, _bucketName, path, cancellationToken);
+							TransferUtilityUploadRequest uploadRequest = new TransferUtilityUploadRequest()
+							{
+								BucketName = _bucketName,
+								Key = path,
+								FilePath = filePath,
+								DisablePayloadSigning = _settings.CurrentValue.DisablePayloadChecksums,
+								DisableDefaultChecksumValidation = _settings.CurrentValue.DisablePayloadChecksums,
+							};
+							await utility.UploadAsync(uploadRequest, cancellationToken);
 						}
 						else
 						{
 							// reset the stream position so that we upload the entire stream each attempt
 							localStream.Position = 0;
-							await utility.UploadAsync(localStream, _bucketName, path, cancellationToken);
+
+							TransferUtilityUploadRequest uploadRequest = new TransferUtilityUploadRequest()
+							{
+								BucketName = _bucketName,
+								Key = path,
+								InputStream = stream,
+								DisablePayloadSigning = _settings.CurrentValue.DisablePayloadChecksums,
+								DisableDefaultChecksumValidation = _settings.CurrentValue.DisablePayloadChecksums,
+							};
+							await utility.UploadAsync(uploadRequest, cancellationToken);
 						}
 
 						break;
diff --git a/Engine/Source/Programs/UnrealCloudDDC/Jupiter/Settings.cs b/Engine/Source/Programs/UnrealCloudDDC/Jupiter/Settings.cs
--- a/Engine/Source/Programs/UnrealCloudDDC/Jupiter/Settings.cs
+++ b/Engine/Source/Programs/UnrealCloudDDC/Jupiter/Settings.cs
@@ -237,6 +237,12 @@ namespace Jupiter
 		/// Larger value means more RAM used but less load on the filesystem.
 		/// </summary>
 		public int MultiPartMaxMemoryBufferSize { get; set; } = int.MaxValue;
+
+		/// <summary>
+		/// Set this to disable payload signing and checksum generation	 during S3 uploads, saves cpu load when uploading large objects at some consistency risk and a requirement to use https connections see
+		/// https://docs.aws.amazon.com/sdkfornet/v3/apidocs/items/S3/TTransferUtilityUploadRequest.html
+		/// </summary>
+		public bool? DisablePayloadChecksums { get; set; } = null;
 	}
 
 	public class GCSettings
