Cloud DDC - Fixed issue with retries of uploads due to bad merge in change 45477634

diff --git a/Engine/Source/Programs/UnrealCloudDDC/Jupiter/Implementation/Blob/AmazonS3Store.cs b/Engine/Source/Programs/UnrealCloudDDC/Jupiter/Implementation/Blob/AmazonS3Store.cs
--- a/Engine/Source/Programs/UnrealCloudDDC/Jupiter/Implementation/Blob/AmazonS3Store.cs
+++ b/Engine/Source/Programs/UnrealCloudDDC/Jupiter/Implementation/Blob/AmazonS3Store.cs
@@ -575,7 +575,7 @@ namespace Jupiter.Implementation
 							{
 								BucketName = _bucketName,
 								Key = path,
-								InputStream = stream,
+								InputStream = localStream,
 								DisablePayloadSigning = _settings.CurrentValue.DisablePayloadChecksums,
 								DisableDefaultChecksumValidation = _settings.CurrentValue.DisablePayloadChecksums,
 							};
