Cloud DDC - Fixed issue with buffering of file when doing multipart uploads and then retrying failed operations

diff --git a/Engine/Source/Programs/UnrealCloudDDC/Jupiter/Implementation/Blob/AmazonS3Store.cs b/Engine/Source/Programs/UnrealCloudDDC/Jupiter/Implementation/Blob/AmazonS3Store.cs
--- a/Engine/Source/Programs/UnrealCloudDDC/Jupiter/Implementation/Blob/AmazonS3Store.cs
+++ b/Engine/Source/Programs/UnrealCloudDDC/Jupiter/Implementation/Blob/AmazonS3Store.cs
@@ -483,10 +483,9 @@ namespace Jupiter.Implementation
 
 		private async Task WriteMultipartAsync(string path, Stream stream, CancellationToken cancellationToken)
 		{
-			FilesystemBufferedPayload? payload = null;
+			IBufferedPayload? payload = null;
 			try
 			{
-				await using MemoryStream localStream = new MemoryStream();
 				string? filePath = null;
 				if (stream is FileStream fileStream)
 				{
@@ -501,14 +500,15 @@ namespace Jupiter.Implementation
 						await using Stream writableStream = writer.GetWritableStream();
 						await stream.CopyToAsync(writableStream, cancellationToken);
 					}
-					payload = writer.Done();
+					FilesystemBufferedPayload filePayload = writer.Done();
+					payload = filePayload;
 
-					filePath = payload.TempFile.FullName;
+					filePath = filePayload.TempFile.FullName;
 				}
 				else
 				{
 					// copy the stream as TransferUtility will dispose this stream if any issue happens making us unable to retry
-					await stream.CopyToAsync(localStream, cancellationToken);
+					payload = await _payloadFactory.CreateFromStreamAsync(stream, stream.Length, "s3-upload", cancellationToken);
 				}
 
 				const int MaxRetryAttempts = 3;
@@ -531,9 +531,7 @@ namespace Jupiter.Implementation
 						}
 						else
 						{
-							// reset the stream position so that we upload the entire stream each attempt
-							localStream.Position = 0;
-
+							await using Stream localStream = payload!.GetStream();
 							TransferUtilityUploadRequest uploadRequest = new TransferUtilityUploadRequest()
 							{
 								BucketName = _bucketName,
