Cloud DDC - Running dotnet monitor as root - with version 8 of this container they no longer do so by default which makes it unable to inspect cloud ddc as that is still runnig as root (something which we stop doing)

diff --git a/Engine/Source/Programs/UnrealCloudDDC/Helm/UnrealCloudDDC/templates/deployment.yaml b/Engine/Source/Programs/UnrealCloudDDC/Helm/UnrealCloudDDC/templates/deployment.yaml
--- a/Engine/Source/Programs/UnrealCloudDDC/Helm/UnrealCloudDDC/templates/deployment.yaml
+++ b/Engine/Source/Programs/UnrealCloudDDC/Helm/UnrealCloudDDC/templates/deployment.yaml
@@ -175,6 +175,11 @@ spec:
     {{- if .Values.useDotnetMonitorSidecar }}
       - name: sidecar
         image: mcr.microsoft.com/dotnet/monitor:8.0.5
+        securityContext:
+          # since we run cloud ddc as root we need to run dotnet monitor as root
+          runAsNonRoot: false
+          runAsUser: 0
+          runAsGroup: 0
         env:
         - name: DOTNETMONITOR_DiagnosticPort__ConnectionMode
           value: Listen
