Cloud DDC - Explicitly set region-specific S3 ServiceURL

diff --git a/Engine/Source/Programs/UnrealCloudDDC/Jupiter/Implementation/Blob/AmazonS3Store.cs b/Engine/Source/Programs/UnrealCloudDDC/Jupiter/Implementation/Blob/AmazonS3Store.cs
--- a/Engine/Source/Programs/UnrealCloudDDC/Jupiter/Implementation/Blob/AmazonS3Store.cs
+++ b/Engine/Source/Programs/UnrealCloudDDC/Jupiter/Implementation/Blob/AmazonS3Store.cs
@@ -1097,10 +1097,9 @@ namespace Jupiter.Implementation
 			{
 				awsOptions.DefaultClientConfig.ServiceURL = s3Settings.ConnectionString;
 			}
-
-			if (region != null)
+			else if (region != null)
 			{
-				awsOptions.Region = RegionEndpoint.GetBySystemName(region);
+				awsOptions.DefaultClientConfig.ServiceURL = $"https://s3.{region}.amazonaws.com";
 			}
 
 			awsOptions.Credentials = awsCredentials;
