Cloud DDC - Capping count of blocks returned at 8000 which is already very large. This helps avoid issue were we try to build a cb object that is larger then 2GB causing overflow exceptions.

diff --git a/Engine/Source/Programs/UnrealCloudDDC/Jupiter/Controllers/BuildsController.cs b/Engine/Source/Programs/UnrealCloudDDC/Jupiter/Controllers/BuildsController.cs
--- a/Engine/Source/Programs/UnrealCloudDDC/Jupiter/Controllers/BuildsController.cs
+++ b/Engine/Source/Programs/UnrealCloudDDC/Jupiter/Controllers/BuildsController.cs
@@ -1619,6 +1619,9 @@ namespace Jupiter.Controllers
 				return NotFound($"Build {buildId} in bucket {bucket} and namespace {ns}");
 			}
 
+			// cap the count at 8000 as objects can become to large to fit into memory otherwise
+			count = Math.Min(count, 8000);
+
 			async Task<bool> WriteBlockMetadata(BlockContext context, CbWriter writer)
 			{
 				using TelemetrySpan scope = _tracer.StartActiveSpan("list.blocks").SetAttribute("operation.name", "list.blocks");
