Cloud DDC - Setting diagnostics port when diagnostics sidecar is enabled

diff --git a/Engine/Source/Programs/UnrealCloudDDC/Helm/UnrealCloudDDC/templates/deployment.yaml b/Engine/Source/Programs/UnrealCloudDDC/Helm/UnrealCloudDDC/templates/deployment.yaml
--- a/Engine/Source/Programs/UnrealCloudDDC/Helm/UnrealCloudDDC/templates/deployment.yaml
+++ b/Engine/Source/Programs/UnrealCloudDDC/Helm/UnrealCloudDDC/templates/deployment.yaml
@@ -94,6 +94,10 @@ spec:
           valueFrom:
             fieldRef:
               fieldPath: metadata.name
+        {{- if .Values.useDotnetMonitorSidecar }}
+        - name: DOTNET_DiagnosticPorts
+          value: /tmp/dotnet-monitor.sock
+        {{- end }}
         {{- if .Values.env }}
         {{- include "common.tplvalues.render" ( dict "value" .Values.env "context" $ ) | nindent 8 }}
         {{- end }}
