Cloud DDC - Fixed issue were retry of blob replication didn't happen if thrown as a aggregate exception

diff --git a/Engine/Source/Programs/UnrealCloudDDC/Jupiter/Implementation/Replication/BlobsReplicator.cs b/Engine/Source/Programs/UnrealCloudDDC/Jupiter/Implementation/Replication/BlobsReplicator.cs
--- a/Engine/Source/Programs/UnrealCloudDDC/Jupiter/Implementation/Replication/BlobsReplicator.cs
+++ b/Engine/Source/Programs/UnrealCloudDDC/Jupiter/Implementation/Replication/BlobsReplicator.cs
@@ -524,6 +524,17 @@ namespace Jupiter.Implementation
 					await _blobService.PutObjectAsync(ns, payload, blob, bucketHint: bucketHint, cancellationToken: cancellationToken);
 					break;
 				}
+				catch (AggregateException e)
+				{
+					Exception baseException = e.GetBaseException();
+					if (baseException is not HttpIOException)
+					{
+						// unknown exception
+						throw;
+					}
+
+					lastException = baseException;
+				}
 				catch (HttpRequestException e)
 				{
 					// rethrow unknown exceptions
