Cloud DDC - Further lowered max block size to 6000 this time to try and fix issue with block response sometimes being larger then 2GB

diff --git a/Engine/Source/Programs/UnrealCloudDDC/Jupiter/Controllers/BuildsController.cs b/Engine/Source/Programs/UnrealCloudDDC/Jupiter/Controllers/BuildsController.cs
--- a/Engine/Source/Programs/UnrealCloudDDC/Jupiter/Controllers/BuildsController.cs
+++ b/Engine/Source/Programs/UnrealCloudDDC/Jupiter/Controllers/BuildsController.cs
@@ -1619,8 +1619,8 @@ namespace Jupiter.Controllers
 				return NotFound($"Build {buildId} in bucket {bucket} and namespace {ns}");
 			}
 
-			// cap the count at 8000 as objects can become to large to fit into memory otherwise
-			count = Math.Min(count, 8000);
+			// cap the count at 6000 as objects can become too large to fit into memory otherwise
+			count = Math.Min(count, 6000);
 
 			async Task<bool> WriteBlockMetadata(BlockContext context, CbWriter writer)
 			{
