Cloud DDC - Update dotnet monitoring sidecar configuration for the worker to enable it to be used to grab traces from the worker as well

diff --git a/Engine/Source/Programs/UnrealCloudDDC/Helm/UnrealCloudDDC/templates/deployment-worker.yaml b/Engine/Source/Programs/UnrealCloudDDC/Helm/UnrealCloudDDC/templates/deployment-worker.yaml
--- a/Engine/Source/Programs/UnrealCloudDDC/Helm/UnrealCloudDDC/templates/deployment-worker.yaml
+++ b/Engine/Source/Programs/UnrealCloudDDC/Helm/UnrealCloudDDC/templates/deployment-worker.yaml
@@ -67,6 +67,10 @@ spec:
           valueFrom:
             fieldRef:
               fieldPath: metadata.namespace
+        {{- if .Values.useDotnetMonitorSidecar }}
+        - name: DOTNET_DiagnosticPorts
+          value: /tmp/dotnet-monitor.sock
+        {{- end }}
         {{- if $worker.env }}
         {{- include "common.tplvalues.render" ( dict "value" $worker.env "context" $ ) | nindent 8 }}
         {{- end }}
@@ -114,11 +118,32 @@ spec:
           {{- end }}
     {{- if $worker.useDotnetMonitorSidecar }}
       - name: sidecar
-        image: mcr.microsoft.com/dotnet/monitor:6.0.2
-        args: ["--urls", "http://*:52323", "--no-auth"]
+        image: mcr.microsoft.com/dotnet/monitor:8.0.5
+        securityContext:
+          # since we run cloud ddc as root we need to run dotnet monitor as root
+          runAsNonRoot: false
+          runAsUser: 0
+          runAsGroup: 0
+        env:
+        - name: DOTNETMONITOR_DiagnosticPort__ConnectionMode
+          value: Listen
+        - name: DOTNETMONITOR_Storage__DefaultSharedPath
+          value: /tmp
+        - name: DOTNETMONITOR_Urls
+          value: http://localhost:52323
+        - name: DOTNETMONITOR_Metrics__Endpoints
+          value: http://+:52325
+        args: ["collect", "--temp-apikey"]
         volumeMounts:
           - name: diagnostics
             mountPath: /tmp
+        resources:
+          requests:
+            memory: "32Mi"
+            cpu: "50m"
+          limits:
+            memory: "256Mi"
+            cpu: "250m"
     {{- end }}
       volumes:
       - name: config-volume
