Cloud DDC - Fixed issue were dotnet sidecar was partially enabled for workers when  enabled for the web api causing the workers to stop on startup

diff --git a/Engine/Source/Programs/UnrealCloudDDC/Helm/UnrealCloudDDC/templates/deployment-worker.yaml b/Engine/Source/Programs/UnrealCloudDDC/Helm/UnrealCloudDDC/templates/deployment-worker.yaml
--- a/Engine/Source/Programs/UnrealCloudDDC/Helm/UnrealCloudDDC/templates/deployment-worker.yaml
+++ b/Engine/Source/Programs/UnrealCloudDDC/Helm/UnrealCloudDDC/templates/deployment-worker.yaml
@@ -67,7 +67,7 @@ spec:
           valueFrom:
             fieldRef:
               fieldPath: metadata.namespace
-        {{- if .Values.useDotnetMonitorSidecar }}
+        {{- if $worker.useDotnetMonitorSidecar }}
         - name: DOTNET_DiagnosticPorts
           value: /tmp/dotnet-monitor.sock
         {{- end }}
