Cloud DDC - Added more logging when test mocks fail to determine what is missing in the mock

diff --git a/Engine/Source/Programs/UnrealCloudDDC/Jupiter/Implementation/Replication/BlobsReplicator.cs b/Engine/Source/Programs/UnrealCloudDDC/Jupiter/Implementation/Replication/BlobsReplicator.cs
--- a/Engine/Source/Programs/UnrealCloudDDC/Jupiter/Implementation/Replication/BlobsReplicator.cs
+++ b/Engine/Source/Programs/UnrealCloudDDC/Jupiter/Implementation/Replication/BlobsReplicator.cs
@@ -365,6 +365,11 @@ namespace Jupiter.Implementation
 
 					lastException = e;
 				}
+				catch (InvalidOperationException)
+				{
+					_logger.LogError("Invalid operation when calling requesting blobs for {Namespace} {RefBucket}", ns, refBucket);
+					throw;
+				}
 			}
 
 			if (response == null)
