Cloud DDC - Removed cap on number of objects returned in listBlocks as its not resolving the error we see and instead added more logging to try and understand how the object becomes so large

diff --git a/Engine/Source/Programs/UnrealCloudDDC/Jupiter/Controllers/BuildsController.cs b/Engine/Source/Programs/UnrealCloudDDC/Jupiter/Controllers/BuildsController.cs
--- a/Engine/Source/Programs/UnrealCloudDDC/Jupiter/Controllers/BuildsController.cs
+++ b/Engine/Source/Programs/UnrealCloudDDC/Jupiter/Controllers/BuildsController.cs
@@ -1619,9 +1619,6 @@ namespace Jupiter.Controllers
 				return NotFound($"Build {buildId} in bucket {bucket} and namespace {ns}");
 			}
 
-			// cap the count at 6000 as objects can become too large to fit into memory otherwise
-			count = Math.Min(count, 6000);
-
 			async Task<bool> WriteBlockMetadata(BlockContext context, CbWriter writer)
 			{
 				using TelemetrySpan scope = _tracer.StartActiveSpan("list.blocks").SetAttribute("operation.name", "list.blocks");
@@ -1708,6 +1705,7 @@ namespace Jupiter.Controllers
 			{
 				case CustomMediaTypeNames.UnrealCompactBinary:
 					{
+						Tracer.CurrentSpan.SetAttribute("response-length", responseObject.GetSize());
 						Response.ContentType = CustomMediaTypeNames.UnrealCompactBinary;
 						Response.StatusCode = 200;
 						await Response.Body.WriteAsync(responseObject.ToByteArray());
