Cloud DDC - Fixed handling of http exceptions during blob replication, should be retrying

diff --git a/Engine/Source/Programs/UnrealCloudDDC/Jupiter/Implementation/Replication/BlobsReplicator.cs b/Engine/Source/Programs/UnrealCloudDDC/Jupiter/Implementation/Replication/BlobsReplicator.cs
--- a/Engine/Source/Programs/UnrealCloudDDC/Jupiter/Implementation/Replication/BlobsReplicator.cs
+++ b/Engine/Source/Programs/UnrealCloudDDC/Jupiter/Implementation/Replication/BlobsReplicator.cs
@@ -531,14 +531,12 @@ namespace Jupiter.Implementation
 				}
 				catch (AggregateException e)
 				{
-					Exception baseException = e.GetBaseException();
-					if (baseException is not HttpIOException)
+					if (e.InnerException is not IOException)
 					{
-						// unknown exception
 						throw;
 					}
 
-					lastException = baseException;
+					lastException = e;
 				}
 				catch (HttpRequestException e)
 				{
