Cloud DDC - Ensure replicator running state is correctly reverted when failure or cancellation occurs retrieving replication log. Explicitly handle cancellations so other replicators are checked.

diff --git a/Engine/Source/Programs/UnrealCloudDDC/Jupiter/Implementation/Replication/ReplicationService.cs b/Engine/Source/Programs/UnrealCloudDDC/Jupiter/Implementation/Replication/ReplicationService.cs
--- a/Engine/Source/Programs/UnrealCloudDDC/Jupiter/Implementation/Replication/ReplicationService.cs
+++ b/Engine/Source/Programs/UnrealCloudDDC/Jupiter/Implementation/Replication/ReplicationService.cs
@@ -110,6 +110,13 @@ namespace Jupiter.Implementation
 							continue;
 						}
 
+						if (replicationTask.IsCanceled)
+						{
+							// Similarly, we check for cancellations but make sure all replicators get to run
+							_logger.LogWarning("Replicator {Name} was canceled", replicator.Info.ReplicatorName);
+							continue;
+						}
+
 						DateTime time = DateTime.Now;
 						_logger.LogInformation("Joining replication task for replicator {Name}", replicator.Info.ReplicatorName);
 						await replicationTask;
