Cloud DDC - Downgrade vendored bitnami/common Helm library for compatibility with existing templates

diff --git a/Engine/Source/Programs/UnrealCloudDDC/Helm/UnrealCloudDDC/charts/common/CHANGELOG.md b/Engine/Source/Programs/UnrealCloudDDC/Helm/UnrealCloudDDC/charts/common/CHANGELOG.md
--- a/Engine/Source/Programs/UnrealCloudDDC/Helm/UnrealCloudDDC/charts/common/CHANGELOG.md
+++ b/Engine/Source/Programs/UnrealCloudDDC/Helm/UnrealCloudDDC/charts/common/CHANGELOG.md
@@ -1,25 +1,8 @@
 # Changelog
 
-## 2.31.4 (2025-08-12)
-
-* [bitnami/common] feat: adapt common.errors.insecureImages for BSI ([#35751]([REDACTED_PROFILE_URL]/charts/pull/35751))
-
-## <small>2.31.3 (2025-06-12)</small>
-
-* [bitnami/common] bugfix: common.capabilities.vpa.apiVersion context (#34372) ([440d610]([REDACTED_PROFILE_URL]/charts/commit/440d6101d0be6e4a37b5f04c6c1ced414d632bfb)), closes [#34372]([REDACTED_PROFILE_URL]/charts/issues/34372)
-
-## <small>2.31.2 (2025-05-20)</small>
-
-* [bitnami/common] Prevent release name from breaking DNS naming specification (#33773) ([07f5233]([REDACTED_PROFILE_URL]/charts/commit/07f523329cad2dfe999b62dc45aa0072541f90be)), closes [#33773]([REDACTED_PROFILE_URL]/charts/issues/33773)
-* [bitnami/kubeapps] Deprecation followup (#33579) ([77e312c]([REDACTED_PROFILE_URL]/charts/commit/77e312c1772d4d7c4dc5d3ac0e80f4e452e3a062)), closes [#33579]([REDACTED_PROFILE_URL]/charts/issues/33579)
-
-## <small>2.31.1 (2025-05-07)</small>
-
-* [bitnami/common] chore: :recycle: Remove unused helpers (#33496) ([d96e16a]([REDACTED_PROFILE_URL]/charts/commit/d96e16aa1e2f1d42de829330f5c0dbc1c4666493)), closes [#33496]([REDACTED_PROFILE_URL]/charts/issues/33496)
-
 ## 2.31.0 (2025-05-05)
 
-* [bitnami/common] chore: :recycle: Remove deprecated APIs (<1.23.0) (#33320) ([3826a9e]([REDACTED_PROFILE_URL]/charts/commit/3826a9e1488c12545f11cf8cb1a11d23daf90602)), closes [#33320]([REDACTED_PROFILE_URL]/charts/issues/33320)
+* [bitnami/common] chore: :recycle: Remove deprecated APIs (<1.23.0) ([#33320]([REDACTED_PROFILE_URL]/charts/pull/33320))
 
 ## <small>2.30.2 (2025-04-30)</small>
 
diff --git a/Engine/Source/Programs/UnrealCloudDDC/Helm/UnrealCloudDDC/charts/common/Chart.yaml b/Engine/Source/Programs/UnrealCloudDDC/Helm/UnrealCloudDDC/charts/common/Chart.yaml
--- a/Engine/Source/Programs/UnrealCloudDDC/Helm/UnrealCloudDDC/charts/common/Chart.yaml
+++ b/Engine/Source/Programs/UnrealCloudDDC/Helm/UnrealCloudDDC/charts/common/Chart.yaml
@@ -2,10 +2,11 @@
 # SPDX-License-Identifier: APACHE-2.0
 
 annotations:
+  category: Infrastructure
   licenses: Apache-2.0
 apiVersion: v2
 # Please make sure that version and appVersion are always the same.
-appVersion: 2.31.4
+appVersion: 2.31.0
 description: A Library Helm Chart for grouping common logic between bitnami charts. This chart is not deployable by itself.
 home: https://bitnami.com
 icon: https://dyltqmyl993wv.cloudfront.net/downloads/logos/bitnami-mark.png
@@ -22,4 +23,4 @@ name: common
 sources:
   - [REDACTED_PROFILE_URL]/charts/tree/main/bitnami/common
 type: library
-version: 2.31.5
\ No newline at end of file
+version: 2.31.0
\ No newline at end of file
diff --git a/Engine/Source/Programs/UnrealCloudDDC/Helm/UnrealCloudDDC/charts/common/README.md b/Engine/Source/Programs/UnrealCloudDDC/Helm/UnrealCloudDDC/charts/common/README.md
--- a/Engine/Source/Programs/UnrealCloudDDC/Helm/UnrealCloudDDC/charts/common/README.md
+++ b/Engine/Source/Programs/UnrealCloudDDC/Helm/UnrealCloudDDC/charts/common/README.md
@@ -19,28 +19,19 @@ helm dependency update
 apiVersion: v1
 kind: ConfigMap
 metadata:
-  name: {{ include "common.names.fullname" . }}
+	name: {{ include "common.names.fullname" . }}
 data:
-  myvalue: "Hello World"
+	myvalue: "Hello World"
 ```
 
 Looking to use our applications in production? Try [VMware Tanzu Application Catalog](https://bitnami.com/enterprise), the commercial edition of the Bitnami catalog.
 
-## ⚠️ Important Notice: Upcoming changes to the Bitnami Catalog
-
-Beginning August 28th, 2025, Bitnami will evolve its public catalog to offer a curated set of hardened, security-focused images under the new [Bitnami Secure Images initiative](https://news.broadcom.com/app-dev/broadcom-introduces-bitnami-secure-images-for-production-ready-containerized-applications). As part of this transition:
-
-- Granting community users access for the first time to security-optimized versions of popular container images.
-- Bitnami will begin deprecating support for non-hardened, Debian-based software images in its free tier and will gradually remove non-latest tags from the public catalog. As a result, community users will have access to a reduced number of hardened images. These images are published only under the “latest” tag and are intended for development purposes
-- Starting August 28th, over two weeks, all existing container images, including older or versioned tags (e.g., 2.50.0, 10.6), will be migrated from the public catalog (docker.io/bitnami) to the “Bitnami Legacy” repository (docker.io/bitnamilegacy), where they will no longer receive updates.
-- For production workloads and long-term support, users are encouraged to adopt Bitnami Secure Images, which include hardened containers, smaller attack surfaces, CVE transparency (via VEX/KEV), SBOMs, and enterprise support.
-
-These changes aim to improve the security posture of all Bitnami users by promoting best practices for software supply chain integrity and up-to-date deployments. For more details, visit the [Bitnami Secure Images announcement]([REDACTED_PROFILE_URL]/containers/issues/83267).
-
 ## Introduction
 
 This chart provides a common template helpers which can be used to develop new charts using [Helm](https://helm.sh) package manager.
 
+Bitnami charts can be used with [Kubeapps](https://kubeapps.dev/) for deployment and management of Helm Charts in clusters.
+
 ## Prerequisites
 
 - Kubernetes 1.23+
@@ -71,6 +62,7 @@ The following table lists the helpers available in the library which are scoped
 | `common.capabilities.job.apiVersion`                      | Return the appropriate apiVersion for job.                                                     | `.` Chart context                       |
 | `common.capabilities.cronjob.apiVersion`                  | Return the appropriate apiVersion for cronjob.                                                 | `.` Chart context                       |
 | `common.capabilities.daemonset.apiVersion`                | Return the appropriate apiVersion for daemonset.                                               | `.` Chart context                       |
+| `common.capabilities.cronjob.apiVersion`                  | Return the appropriate apiVersion for cronjob.                                                 | `.` Chart context                       |
 | `common.capabilities.deployment.apiVersion`               | Return the appropriate apiVersion for deployment.                                              | `.` Chart context                       |
 | `common.capabilities.statefulset.apiVersion`              | Return the appropriate apiVersion for statefulset.                                             | `.` Chart context                       |
 | `common.capabilities.ingress.apiVersion`                  | Return the appropriate apiVersion for ingress.                                                 | `.` Chart context                       |
@@ -115,6 +107,8 @@ The following table lists the helpers available in the library which are scoped
 | Helper identifier                         | Description                                                                                                       | Expected Input                                                                                                                                                                   |
 | ----------------------------------------- | ----------------------------------------------------------------------------------------------------------------- | -------------------------------------------------------------------------------------------------------------------------------------------------------------------------------- |
 | `common.ingress.backend`                  | Generate a proper Ingress backend entry depending on the API version                                              | `dict "serviceName" "foo" "servicePort" "bar"`, see the [Ingress deprecation notice](https://kubernetes.io/blog/2019/07/18/api-deprecations-in-1-16/) for the syntax differences |
+| `common.ingress.supportsPathType`         | Prints "true" if the pathType field is supported                                                                  | `.` Chart context                                                                                                                                                                |
+| `common.ingress.supportsIngressClassname` | Prints "true" if the ingressClassname field is supported                                                          | `.` Chart context                                                                                                                                                                |
 | `common.ingress.certManagerRequest`       | Prints "true" if required cert-manager annotations for TLS signed certificates are set in the Ingress annotations | `dict "annotations" .Values.path.to.the.ingress.annotations`                                                                                                                     |
 
 ### Labels
