Cloud DDC - Added alternative implementation of last access service that is a background service which hopefully works better with the long running async task we spawn

Removed unused startLastAccessRollup admin method as it was impractical to use and assumed access to the background service which is just a bit odd.

diff --git a/Engine/Source/Programs/UnrealCloudDDC/Helm/UnrealCloudDDC/templates/configmap.yaml b/Engine/Source/Programs/UnrealCloudDDC/Helm/UnrealCloudDDC/templates/configmap.yaml
--- a/Engine/Source/Programs/UnrealCloudDDC/Helm/UnrealCloudDDC/templates/configmap.yaml
+++ b/Engine/Source/Programs/UnrealCloudDDC/Helm/UnrealCloudDDC/templates/configmap.yaml
@@ -52,6 +52,10 @@ UnrealCloudDDC:
   EnableInlineSmallBlobs: {{ $config.EnableInlineSmallBlobs }}
 {{- end }}
 
+{{- if hasKey $config "UseLastAccessBackgroundService" }}
+  UseLastAccessBackgroundService: {{ $config.UseLastAccessBackgroundService }}
+{{- end }}
+
 {{ if $config.StorageImplementations }}
   StorageImplementations:
   {{- range $config.StorageImplementations -}}
diff --git a/Engine/Source/Programs/UnrealCloudDDC/Jupiter/Controllers/AdminController.cs b/Engine/Source/Programs/UnrealCloudDDC/Jupiter/Controllers/AdminController.cs
--- a/Engine/Source/Programs/UnrealCloudDDC/Jupiter/Controllers/AdminController.cs
+++ b/Engine/Source/Programs/UnrealCloudDDC/Jupiter/Controllers/AdminController.cs
@@ -18,44 +18,17 @@ namespace Jupiter.Controllers
 	[Authorize]
 	public class AdminController : Controller
 	{
-		private readonly LastAccessServiceReferences _lastAccessService;
 		private readonly IRefCleanup _refCleanup;
 		private readonly IConfiguration _configuration;
 		private readonly IRequestHelper _requestHelper;
 
-		public AdminController(LastAccessServiceReferences lastAccessService, IRefCleanup refCleanup, IConfiguration configuration, IRequestHelper requestHelper)
+		public AdminController(IRefCleanup refCleanup, IConfiguration configuration, IRequestHelper requestHelper)
 		{
-			_lastAccessService = lastAccessService;
 			_refCleanup = refCleanup;
 			_configuration = configuration;
 			_requestHelper = requestHelper;
 		}
-
-		/// <summary>
-		/// Report last access times to refs store
-		/// </summary>
-		/// <remarks>
-		/// Manually triggers a aggregation of the last access records and commits these to the refs store. This is done automatically so the only reason to use this endpoint is for debugging purposes.
-		/// </remarks>
-		/// <returns></returns>
-		[HttpPost("startLastAccessRollup")]
-		[ProducesResponseType(type: typeof(UpdatedRecordsResponse), 200)]
-		public async Task<IActionResult> StartLastAccessRollupAsync()
-		{
-			ActionResult? result = await _requestHelper.HasAccessForGlobalOperationsAsync(User, new[] { JupiterAclAction.AdminAction });
-			if (result != null)
-			{
-				return result;
-			}
-
-			List<(LastAccessRecord, DateTime)>? updatedRecords = await _lastAccessService.ProcessLastAccessRecordsAsync(HttpContext.RequestAborted);
-
-			return Ok(new UpdatedRecordsResponse(
-				updatedRecords?.Select(tuple => new UpdatedRecordsResponse.UpdatedRecord(tuple.Item1, tuple.Item2))
-					.ToList() ??
-				new List<UpdatedRecordsResponse.UpdatedRecord>()));
-		}
-
+		
 		/// <summary>
 		/// Manually run the refs cleanup
 		/// </summary>
diff --git a/Engine/Source/Programs/UnrealCloudDDC/Jupiter/Implementation/LastAccess/LastAccessServiceReferences.cs b/Engine/Source/Programs/UnrealCloudDDC/Jupiter/Implementation/LastAccess/LastAccessServiceReferences.cs
--- a/Engine/Source/Programs/UnrealCloudDDC/Jupiter/Implementation/LastAccess/LastAccessServiceReferences.cs
+++ b/Engine/Source/Programs/UnrealCloudDDC/Jupiter/Implementation/LastAccess/LastAccessServiceReferences.cs
@@ -122,4 +122,78 @@ namespace Jupiter.Implementation
 			}
 		}
 	}
+
+	// ReSharper disable once ClassNeverInstantiated.Global
+	public class LastAccessBackgroundService : BackgroundService
+	{
+		private readonly ILastAccessCache<LastAccessRecord> _lastAccessCacheRecord;
+		private readonly IReferencesStore _referencesStore;
+		private readonly INamespacePolicyResolver _namespacePolicyResolver;
+		private readonly Tracer _tracer;
+		private readonly ILogger _logger;
+		private readonly UnrealCloudDDCSettings _settings;
+
+		public LastAccessBackgroundService(IOptionsMonitor<UnrealCloudDDCSettings> settings, ILastAccessCache<LastAccessRecord> lastAccessCache, IReferencesStore referencesStore, INamespacePolicyResolver namespacePolicyResolver, Tracer tracer, ILogger<LastAccessServiceReferences> logger)
+		{
+			_lastAccessCacheRecord = lastAccessCache;
+			_referencesStore = referencesStore;
+			_namespacePolicyResolver = namespacePolicyResolver;
+			_tracer = tracer;
+			_settings = settings.CurrentValue;
+			_logger = logger;
+		}
+
+		protected override async Task ExecuteAsync(CancellationToken stoppingToken)
+		{
+			while (!stoppingToken.IsCancellationRequested)
+			{
+				try
+				{
+					await ProcessLastAccessRecordsAsync(stoppingToken);
+				}
+				catch (Exception e)
+				{
+					_logger.LogError(e, "Exception thrown while processing last access records");
+					Tracer.CurrentSpan.SetStatus(Status.Error);
+					Tracer.CurrentSpan.RecordException(e);
+				}
+
+				await Task.Delay(_settings.LastAccessRollupFrequencySeconds, stoppingToken);
+			}
+		}
+
+		private async Task<List<(LastAccessRecord, DateTime)>> ProcessLastAccessRecordsAsync(CancellationToken cancellationToken)
+		{
+			_logger.LogInformation("Running Last Access Aggregation for refs");
+			List<(LastAccessRecord, DateTime)> records = await _lastAccessCacheRecord.GetLastAccessedRecords();
+			foreach ((LastAccessRecord record, DateTime lastAccessTime) in records)
+			{
+				if (!ShouldTrackLastAccess(record.Namespace))
+				{
+					continue;
+				}
+
+				using TelemetrySpan scope = _tracer.StartActiveSpan("lastAccess.update")
+					.SetAttribute("operation.name", "lastAccess.update")
+					.SetAttribute("resource.name", $"{record.Namespace}:{record.Bucket}.{record.Key}");
+				_logger.LogDebug("Updating last access time to {LastAccessTime} for {Record}", lastAccessTime, record);
+				await _referencesStore.UpdateLastAccessTimeAsync(record.Namespace, record.Bucket, record.Key, lastAccessTime, cancellationToken);
+				// delay 10ms between each record to distribute the load more evenly for the db
+				await Task.Delay(10, cancellationToken);
+			}
+
+			return records;
+		}
+
+		private bool ShouldTrackLastAccess(NamespaceId ns)
+		{
+			if (!_namespacePolicyResolver.GetPoliciesForNs(ns).LastAccessTracking)
+			{
+				return false;
+			}
+
+			NamespacePolicy.StoragePoolGCMethod gcMethod = _namespacePolicyResolver.GetPoliciesForNs(ns).GcMethod ?? NamespacePolicy.StoragePoolGCMethod.LastAccess;
+			return gcMethod == NamespacePolicy.StoragePoolGCMethod.LastAccess;
+		}
+	}
 }
diff --git a/Engine/Source/Programs/UnrealCloudDDC/Jupiter/Jupiter.Startup.cs b/Engine/Source/Programs/UnrealCloudDDC/Jupiter/Jupiter.Startup.cs
--- a/Engine/Source/Programs/UnrealCloudDDC/Jupiter/Jupiter.Startup.cs
+++ b/Engine/Source/Programs/UnrealCloudDDC/Jupiter/Jupiter.Startup.cs
@@ -173,9 +173,6 @@ namespace Jupiter
 			services.AddSingleton(serviceType: typeof(ISecretResolver), typeof(SecretResolver));
 			services.AddSingleton(typeof(IAmazonSecretsManager), CreateAWSSecretsManager);
 
-			services.AddSingleton<LastAccessServiceReferences>();
-			services.AddHostedService<LastAccessServiceReferences>(p => p.GetService<LastAccessServiceReferences>()!);
-
 			services.AddSingleton<IServiceCredentials, ServiceCredentials>(p => ActivatorUtilities.CreateInstance<ServiceCredentials>(p));
 
 			services.AddSingleton<ReplicationService>();
@@ -219,6 +216,17 @@ namespace Jupiter
 				services.AddHostedService<KubernetesLeaderElection>(p => p.GetService<KubernetesLeaderElection>()!);
 			}
 
+			if (settings.UseLastAccessBackgroundService)
+			{
+				services.AddSingleton<LastAccessBackgroundService>();
+				services.AddHostedService<LastAccessBackgroundService>(p => p.GetService<LastAccessBackgroundService>()!);
+			}
+			else
+			{
+				services.AddSingleton<LastAccessServiceReferences>();
+				services.AddHostedService<LastAccessServiceReferences>(p => p.GetService<LastAccessServiceReferences>()!);
+			}
+
 			services.AddSingleton<BlobStoreConsistencyCheckService>();
 			services.AddSingleton(typeof(IPeerServiceDiscovery), ServiceDiscoveryFactory);
 		}
diff --git a/Engine/Source/Programs/UnrealCloudDDC/Jupiter/Settings.cs b/Engine/Source/Programs/UnrealCloudDDC/Jupiter/Settings.cs
--- a/Engine/Source/Programs/UnrealCloudDDC/Jupiter/Settings.cs
+++ b/Engine/Source/Programs/UnrealCloudDDC/Jupiter/Settings.cs
@@ -137,6 +137,11 @@ namespace Jupiter
 		/// Used to force a more parallel resolve than normal (which has no limits). Can be used to go wider then the number of processors which improves how quickly the finalize for a ref can be done, especially useful for refs with a lot of dependencies
 		/// </summary>
 		public int RefFinalizeResolveMaxParallel { get; set; } = Environment.ProcessorCount;
+
+		/// <summary>
+		/// Set to enable using the new last access service which is registered as a background service
+		/// </summary>
+		public bool UseLastAccessBackgroundService { get; set; } = false;
 	}
 
 	public class MongoSettings
