Cloud DDC - Fixed issue were last access background service was running much more frequently then intended

diff --git a/Engine/Source/Programs/UnrealCloudDDC/Jupiter/Implementation/LastAccess/LastAccessServiceReferences.cs b/Engine/Source/Programs/UnrealCloudDDC/Jupiter/Implementation/LastAccess/LastAccessServiceReferences.cs
--- a/Engine/Source/Programs/UnrealCloudDDC/Jupiter/Implementation/LastAccess/LastAccessServiceReferences.cs
+++ b/Engine/Source/Programs/UnrealCloudDDC/Jupiter/Implementation/LastAccess/LastAccessServiceReferences.cs
@@ -158,7 +158,7 @@ namespace Jupiter.Implementation
 					Tracer.CurrentSpan.RecordException(e);
 				}
 
-				await Task.Delay(_settings.LastAccessRollupFrequencySeconds, stoppingToken);
+				await Task.Delay(TimeSpan.FromSeconds(_settings.LastAccessRollupFrequencySeconds), stoppingToken);
 			}
 		}
 
