Cloud DDC - Made the rollup frequency for last access background service more predicable by running at the expected frequency but still not running to frequently if there is no work to do.

diff --git a/Engine/Source/Programs/UnrealCloudDDC/Jupiter/Implementation/LastAccess/LastAccessServiceReferences.cs b/Engine/Source/Programs/UnrealCloudDDC/Jupiter/Implementation/LastAccess/LastAccessServiceReferences.cs
--- a/Engine/Source/Programs/UnrealCloudDDC/Jupiter/Implementation/LastAccess/LastAccessServiceReferences.cs
+++ b/Engine/Source/Programs/UnrealCloudDDC/Jupiter/Implementation/LastAccess/LastAccessServiceReferences.cs
@@ -149,7 +149,10 @@ namespace Jupiter.Implementation
 			{
 				try
 				{
-					await ProcessLastAccessRecordsAsync(stoppingToken);
+					Task delay = Task.Delay(TimeSpan.FromSeconds(_settings.LastAccessRollupFrequencySeconds), stoppingToken);
+					
+					// run the rollout once every rollup frequency at most, most of the time this takes longer then the rollup frequency at which point we just start again
+					await Task.WhenAll(delay, ProcessLastAccessRecordsAsync(stoppingToken));
 				}
 				catch (Exception e)
 				{
@@ -157,8 +160,6 @@ namespace Jupiter.Implementation
 					Tracer.CurrentSpan.SetStatus(Status.Error);
 					Tracer.CurrentSpan.RecordException(e);
 				}
-
-				await Task.Delay(TimeSpan.FromSeconds(_settings.LastAccessRollupFrequencySeconds), stoppingToken);
 			}
 		}
 
