Cloud DDC - Added more debug logging to understand why GetSize of a object returns negative sizes for some objects

diff --git a/Engine/Source/Programs/UnrealCloudDDC/Jupiter/Controllers/BuildsController.cs b/Engine/Source/Programs/UnrealCloudDDC/Jupiter/Controllers/BuildsController.cs
--- a/Engine/Source/Programs/UnrealCloudDDC/Jupiter/Controllers/BuildsController.cs
+++ b/Engine/Source/Programs/UnrealCloudDDC/Jupiter/Controllers/BuildsController.cs
@@ -1646,6 +1646,7 @@ namespace Jupiter.Controllers
 					enumerable = metadataBlobs.OrderByDescending(tuple => tuple.Item1);
 				}
 
+				ulong totalSize = 0;
 				await Parallel.ForEachAsync(enumerable.Skip(skip).Take(count), async (pair, token) =>
 				{
 					(DateTime _, BlobId blockMetadata) = pair;
@@ -1665,6 +1666,10 @@ namespace Jupiter.Controllers
 						lock (writer)
 						{
 							writer.WriteObject(o);
+
+							int objectSize = o.GetSize();
+							Interlocked.Add(ref totalSize, (ulong)objectSize);
+							_logger.LogInformation("Adding object of size {ObjectSize} total size is now {TotalSize}", objectSize, totalSize);
 						}
 
 						Interlocked.Increment(ref countOfBlocksFound);
