Cloud DDC - Return 503 (Service Unavailable) when underlying database is unavailable

diff --git a/Engine/Source/Programs/UnrealCloudDDC/Jupiter.Common/RemoteServiceException.cs b/Engine/Source/Programs/UnrealCloudDDC/Jupiter.Common/RemoteServiceException.cs
new file mode 100644
--- /dev/null
+++ b/Engine/Source/Programs/UnrealCloudDDC/Jupiter.Common/RemoteServiceException.cs
@@ -0,0 +1,13 @@
+// Copyright Epic Games, Inc. All Rights Reserved.
+
+using System;
+
+namespace Jupiter.Common
+{
+	public class RemoteServiceException : Exception
+	{
+		public RemoteServiceException(string message, Exception cause) : base(message, cause)
+		{
+		}
+	}
+}
diff --git a/Engine/Source/Programs/UnrealCloudDDC/Jupiter/Controllers/BlobController.cs b/Engine/Source/Programs/UnrealCloudDDC/Jupiter/Controllers/BlobController.cs
--- a/Engine/Source/Programs/UnrealCloudDDC/Jupiter/Controllers/BlobController.cs
+++ b/Engine/Source/Programs/UnrealCloudDDC/Jupiter/Controllers/BlobController.cs
@@ -84,6 +84,10 @@ namespace Jupiter.Controllers
 			{
 				return e.Result;
 			}
+			catch (RemoteServiceException e)
+			{
+				return Problem(e.Message, null, (int)HttpStatusCode.ServiceUnavailable);
+			}
 		}
 
 		[HttpHead("{ns}/{id}")]
@@ -98,11 +102,19 @@ namespace Jupiter.Controllers
 			{
 				return result;
 			}
-			bool exists = await _storage.ExistsAsync(ns, id, storageLayers);
 
-			if (!exists)
+			try
 			{
-				return NotFound(new ValidationProblemDetails { Title = $"Blob {id} not found" });
+				bool exists = await _storage.ExistsAsync(ns, id, storageLayers);
+
+				if (!exists)
+				{
+					return NotFound(new ValidationProblemDetails { Title = $"Blob {id} not found" });
+				}
+			}
+			catch (RemoteServiceException e)
+			{
+				return Problem(e.Message, null, (int)HttpStatusCode.ServiceUnavailable);
 			}
 
 			return Ok();
@@ -214,6 +226,10 @@ namespace Jupiter.Controllers
 			{
 				return Problem(e.Message, null, (int)HttpStatusCode.RequestTimeout);
 			}
+			catch (RemoteServiceException e)
+			{
+				return Problem(e.Message, null, (int)HttpStatusCode.ServiceUnavailable);
+			}
 		}
 
 		[HttpPost("{ns}")]
@@ -245,6 +261,10 @@ namespace Jupiter.Controllers
 			{
 				return Problem(e.Message, null, (int)HttpStatusCode.RequestTimeout);
 			}
+			catch (RemoteServiceException e)
+			{
+				return Problem(e.Message, null, (int)HttpStatusCode.ServiceUnavailable);
+			}
 		}
 
 		[HttpDelete("{ns}/{id}")]
@@ -258,7 +278,14 @@ namespace Jupiter.Controllers
 				return result;
 			}
 
-			await DeleteImplAsync(ns, id);
+			try
+			{
+				await DeleteImplAsync(ns, id);
+			}
+			catch (RemoteServiceException e)
+			{
+				return Problem(e.Message, null, (int)HttpStatusCode.ServiceUnavailable);
+			}
 
 			return NoContent();
 		}
@@ -273,7 +300,14 @@ namespace Jupiter.Controllers
 				return result;
 			}
 
-			await _storage.DeleteNamespaceAsync(ns, HttpContext.RequestAborted);
+			try
+			{
+				await _storage.DeleteNamespaceAsync(ns, HttpContext.RequestAborted);
+			}
+			catch (RemoteServiceException e)
+			{
+				return Problem(e.Message, null, (int)HttpStatusCode.ServiceUnavailable);
+			}
 
 			return NoContent();
 		}
diff --git a/Engine/Source/Programs/UnrealCloudDDC/Jupiter/Controllers/BuildsController.cs b/Engine/Source/Programs/UnrealCloudDDC/Jupiter/Controllers/BuildsController.cs
--- a/Engine/Source/Programs/UnrealCloudDDC/Jupiter/Controllers/BuildsController.cs
+++ b/Engine/Source/Programs/UnrealCloudDDC/Jupiter/Controllers/BuildsController.cs
@@ -163,8 +163,16 @@ namespace Jupiter.Controllers
 			CbObject mutatedObject = newObjectWriter.ToObject();
 
 			uint ttl = (uint)_namespacePolicyResolver.GetPoliciesForNs(ns).DefaultTTL.TotalSeconds;
-			await _buildStore.PutBuildAsync(ns, bucket, buildId, mutatedObject, ttl);
 
+			try
+			{
+				await _buildStore.PutBuildAsync(ns, bucket, buildId, mutatedObject, ttl);
+			}
+			catch (RemoteServiceException e)
+			{
+				return Problem(e.Message, null, (int)HttpStatusCode.ServiceUnavailable);
+			}
+			
 			return Ok(new PutBuildResponse {ChunkSize = idealChunkSize});
 		}
 
@@ -184,8 +192,16 @@ namespace Jupiter.Controllers
 			Tracer.CurrentSpan.SetAttribute("namespace", ns.ToString());
 
 			uint ttl = (uint)_namespacePolicyResolver.GetPoliciesForNs(ns).DefaultTTL.TotalSeconds;
-			await _buildStore.FinalizeBuildAsync(ns, bucket, buildId, ttl);
 
+			try
+			{
+				await _buildStore.FinalizeBuildAsync(ns, bucket, buildId, ttl);
+			}
+			catch (RemoteServiceException e)
+			{
+				return Problem(e.Message, null, (int)HttpStatusCode.ServiceUnavailable);
+			}
+			
 			return Ok();
 		}
 
@@ -202,13 +218,20 @@ namespace Jupiter.Controllers
 			Tracer.CurrentSpan.SetAttribute("bucket", bucket.ToString());
 			Tracer.CurrentSpan.SetAttribute("namespace", ns.ToString());
 
-			uint? ttl = await _buildStore.GetTTL(ns, bucket, buildId);
+			try
+			{
+				uint? ttl = await _buildStore.GetTTL(ns, bucket, buildId);
 
-			if (!ttl.HasValue)
+				if (!ttl.HasValue)
+				{
+					return NotFound();
+				}
+				return Ok(new GetTTLResponse { TTL = ttl.Value});
+			}
+			catch (RemoteServiceException e)
 			{
-				return NotFound();
+				return Problem(e.Message, null, (int)HttpStatusCode.ServiceUnavailable);
 			}
-			return Ok(new GetTTLResponse { TTL = ttl.Value});
 		}
 
 		[HttpPost("{ns}/{bucket}/{buildId}/updateTTL")]
@@ -226,13 +249,20 @@ namespace Jupiter.Controllers
 
 			LogAuditEntry(HttpMethod.Post, nameof(UpdateTTLAsync), ns, bucket, buildId, null);
 
-			await foreach ((string _, CbObjectId partId) in _buildStore.GetBuildPartsAsync(ns, bucket, buildId))
+			try
 			{
-				RefId refId = RefId.FromName($"{buildId}/{partId}");
-				await _refService.UpdateTTL(ns, bucket, refId, request.TTL);
-			}
+				await foreach ((string _, CbObjectId partId) in _buildStore.GetBuildPartsAsync(ns, bucket, buildId))
+				{
+					RefId refId = RefId.FromName($"{buildId}/{partId}");
+					await _refService.UpdateTTL(ns, bucket, refId, request.TTL);
+				}
 
-			await _buildStore.UpdateTTL(ns, bucket, buildId, request.TTL);
+				await _buildStore.UpdateTTL(ns, bucket, buildId, request.TTL);
+			}
+			catch (RemoteServiceException e)
+			{
+				return Problem(e.Message, null, (int)HttpStatusCode.ServiceUnavailable);
+			}
 
 			return Ok();
 		}
@@ -252,13 +282,20 @@ namespace Jupiter.Controllers
 
 			LogAuditEntry(HttpMethod.Delete, nameof(DeleteBuildAsync), ns, bucket, buildId, null);
 
-			await foreach ((string _, CbObjectId partId) in _buildStore.GetBuildPartsAsync(ns, bucket, buildId))
+			try
 			{
-				RefId refId = RefId.FromName($"{buildId}/{partId}");
-				await _refService.DeleteAsync(ns, bucket, refId);
-			}
+				await foreach ((string _, CbObjectId partId) in _buildStore.GetBuildPartsAsync(ns, bucket, buildId))
+				{
+					RefId refId = RefId.FromName($"{buildId}/{partId}");
+					await _refService.DeleteAsync(ns, bucket, refId);
+				}
 
-			await _buildStore.DeleteBuild(ns, bucket, buildId);
+				await _buildStore.DeleteBuild(ns, bucket, buildId);
+			}
+			catch (RemoteServiceException e)
+			{
+				return Problem(e.Message, null, (int)HttpStatusCode.ServiceUnavailable);
+			}
 
 			return Ok();
 		}
@@ -278,53 +315,60 @@ namespace Jupiter.Controllers
 
 			LogAuditEntry(HttpMethod.Get, nameof(GetBuildAsync), ns, bucket, buildId, null);
 
-			BuildRecord? buildRecord = await _buildStore.GetBuildAsync(ns, bucket, buildId);
-
-			if (buildRecord == null)
-			{
-				return NotFound($"Build \"{buildId}\" in bucket {bucket} and namespace {ns}");
-			}
-
-			// write a new cb object because we want to append the parts field to it for the build parts
-			CbWriter newObjectWriter = new CbWriter();
-			newObjectWriter.BeginObject();
-			foreach (CbField field in buildRecord.BuildObject)
+			try
 			{
-				newObjectWriter.WriteField(field);
-			}
+				BuildRecord? buildRecord = await _buildStore.GetBuildAsync(ns, bucket, buildId);
 
-			// inject the build parts into the object
-			// we know the object fields will all have the same type so make it a uniform object to please compact binary validation
-			newObjectWriter.BeginUniformObject("parts", CbFieldType.ObjectId | CbFieldType.HasFieldName);
-			await foreach ((string partName, CbObjectId partId) in _buildStore.GetBuildPartsAsync(ns, bucket, buildId))
-			{
+				if (buildRecord == null)
+				{
+					return NotFound($"Build \"{buildId}\" in bucket {bucket} and namespace {ns}");
+				}
 
-				newObjectWriter.WriteObjectId(partName, partId);
-			}
-			newObjectWriter.EndObject(); // parts
-			newObjectWriter.EndObject(); // root
-			CbObject newObject = newObjectWriter.ToObject();
-			
-			string responseType = _formatResolver.GetResponseType(Request, format, CustomMediaTypeNames.UnrealCompactBinary);
-			Tracer.CurrentSpan.SetAttribute("response-type", responseType);
+				// write a new cb object because we want to append the parts field to it for the build parts
+				CbWriter newObjectWriter = new CbWriter();
+				newObjectWriter.BeginObject();
+				foreach (CbField field in buildRecord.BuildObject)
+				{
+					newObjectWriter.WriteField(field);
+				}
 
-			switch (responseType)
-			{
-				case CustomMediaTypeNames.UnrealCompactBinary:
+				// inject the build parts into the object
+				// we know the object fields will all have the same type so make it a uniform object to please compact binary validation
+				newObjectWriter.BeginUniformObject("parts", CbFieldType.ObjectId | CbFieldType.HasFieldName);
+				await foreach ((string partName, CbObjectId partId) in _buildStore.GetBuildPartsAsync(ns, bucket, buildId))
 				{
-					Response.ContentType = CustomMediaTypeNames.UnrealCompactBinary;
-					Response.StatusCode = 200;
-					await Response.Body.WriteAsync(newObject.GetView());
-					return new EmptyResult();
+
+					newObjectWriter.WriteObjectId(partName, partId);
 				}
-				case MediaTypeNames.Application.Json:
+				newObjectWriter.EndObject(); // parts
+				newObjectWriter.EndObject(); // root
+				CbObject newObject = newObjectWriter.ToObject();
+				
+				string responseType = _formatResolver.GetResponseType(Request, format, CustomMediaTypeNames.UnrealCompactBinary);
+				Tracer.CurrentSpan.SetAttribute("response-type", responseType);
+
+				switch (responseType)
 				{
-					string s = newObject.ToJson();
-					Response.ContentType = MediaTypeNames.Application.Json;
-					return Ok(s);
+					case CustomMediaTypeNames.UnrealCompactBinary:
+					{
+						Response.ContentType = CustomMediaTypeNames.UnrealCompactBinary;
+						Response.StatusCode = 200;
+						await Response.Body.WriteAsync(newObject.GetView());
+						return new EmptyResult();
+					}
+					case MediaTypeNames.Application.Json:
+					{
+						string s = newObject.ToJson();
+						Response.ContentType = MediaTypeNames.Application.Json;
+						return Ok(s);
+					}
+					default:
+						throw new NotImplementedException($"Unknown response type: {responseType}");
 				}
-				default:
-					throw new NotImplementedException($"Unknown response type: {responseType}");
+			}
+			catch (RemoteServiceException e)
+			{
+				return Problem(e.Message, null, (int)HttpStatusCode.ServiceUnavailable);
 			}
 		}
 
@@ -353,57 +397,66 @@ namespace Jupiter.Controllers
 
 			LogAuditEntry(HttpMethod.Post, nameof(SearchBuildAsync), ns, bucket, null, null);
 
-			await foreach (BuildMetadata build in _buildStore.ListBuildsAsync(ns, bucket, includeTTL))
+			try
 			{
-				if (!build.IsFinalized && !includePartial)
-				{
-					// we ignore non-finalized builds unless explicitly asked to include them
-					continue;
-				}
-				count++;
-				if (count < searchRequest.Options.Skip)
+				await foreach (BuildMetadata build in _buildStore.ListBuildsAsync(ns, bucket, includeTTL))
 				{
-					continue;
-				}
-
-				if (count >= searchRequest.Options.Max)
-				{
-					// we have considered all the builds we are allowed to
-					partialResult = true;
-					break;
-				}
+					if (!build.IsFinalized && !includePartial)
+					{
+						// we ignore non-finalized builds unless explicitly asked to include them
+						continue;
+					}
 
-				timeSpent = DateTime.Now - startTime;
-				if (checkForMaxTime && timeSpent > maxTimeToSpend)
-				{
-					// spent all the time we were given for this query
-					OkObjectResult result = Ok(new SearchResult(foundResults, true));
-					result.StatusCode = (int)HttpStatusCode.RequestTimeout;
-					return result;
-				}
+					count++;
+					if (count < searchRequest.Options.Skip)
+					{
+						continue;
+					}
 
-				bool match = true;
-				foreach (ISearchOp op in searchOps)
-				{
-					CbObject o = build.Metadata;
-					if (!op.Matches(build.BuildId, o, _logger))
+					if (count >= searchRequest.Options.Max)
 					{
-						match = false;
+						// we have considered all the builds we are allowed to
+						partialResult = true;
 						break;
 					}
-				}
 
-				if (match)
-				{
-					foundResults.Add(build);
-				}
+					timeSpent = DateTime.Now - startTime;
+					if (checkForMaxTime && timeSpent > maxTimeToSpend)
+					{
+						// spent all the time we were given for this query
+						OkObjectResult result = Ok(new SearchResult(foundResults, true));
+						result.StatusCode = (int)HttpStatusCode.RequestTimeout;
+						return result;
+					}
 
-				if (foundResults.Count >= searchRequest.Options.Limit)
-				{
-					partialResult = true;
-					break;
+					bool match = true;
+					foreach (ISearchOp op in searchOps)
+					{
+						CbObject o = build.Metadata;
+						if (!op.Matches(build.BuildId, o, _logger))
+						{
+							match = false;
+							break;
+						}
+					}
+
+					if (match)
+					{
+						foundResults.Add(build);
+					}
+
+					if (foundResults.Count >= searchRequest.Options.Limit)
+					{
+						partialResult = true;
+						break;
+					}
 				}
 			}
+			catch (RemoteServiceException e)
+			{
+				return Problem(e.Message, null, (int)HttpStatusCode.ServiceUnavailable);
+			}
+
 			_logger.LogDebug("Search completed in {Namespace} {Bucket} request: {@SearchRequest} Results: {CountOfResults} Considered: {CountOfBuildsConsidered} took {Duration}", ns, bucket, searchRequest, foundResults.Count, count, timeSpent);
 			Tracer.CurrentSpan.SetAttribute("countOfMatches", foundResults.Count);
 			_diagnosticContext.Set("MatchCount", foundResults.Count);
@@ -431,37 +484,47 @@ namespace Jupiter.Controllers
 				return BadRequest($"Invalid regular expression for bucket filtering: {searchBucketRequest.BucketRegex}. Parse error was: {e.Error}");
 			}
 
-			bool lackedAccessToAtLeastOneBucket = false;
-			foreach (BucketId bucket in await _refService.GetBucketsAsync(ns, HttpContext.RequestAborted).ToArrayAsync(HttpContext.RequestAborted))
+			try
 			{
-				if (!regex.IsMatch(bucket.ToString()))
+				bool lackedAccessToAtLeastOneBucket = false;
+				foreach (BucketId bucket in await _refService.GetBucketsAsync(ns, HttpContext.RequestAborted)
+					         .ToArrayAsync(HttpContext.RequestAborted))
 				{
-					// does not match the filter
-					continue;
+					if (!regex.IsMatch(bucket.ToString()))
+					{
+						// does not match the filter
+						continue;
+					}
+
+					// verify we have access to the bucket
+					ActionResult? accessResult = await _requestHelper.HasAccessToScopeAsync(User, Request, new AccessScope(ns, bucket),
+						new[] { JupiterAclAction.EnumerateBucket });
+					if (accessResult == null)
+					{
+						bucketsToSearch.Add(bucket);
+					}
+					else
+					{
+						lackedAccessToAtLeastOneBucket = true;
+					}
 				}
-				// verify we have access to the bucket
-				ActionResult? accessResult = await _requestHelper.HasAccessToScopeAsync(User, Request, new AccessScope(ns, bucket), new[] { JupiterAclAction.EnumerateBucket });
-				if (accessResult == null)
+
+				// if we found no matching buckets that we had access to we raise a 403
+				if (bucketsToSearch.Count == 0 && lackedAccessToAtLeastOneBucket)
 				{
-					bucketsToSearch.Add(bucket);
+					return new ForbidResult();
 				}
-				else
+
+				// no buckets matched
+				if (bucketsToSearch.Count == 0)
 				{
-					lackedAccessToAtLeastOneBucket = true;
+					// it's not an error to search were nothing matches, but it could be useful to know that it was due to the regular expression
+					return Ok(new SearchMultipleBucketsResult());
 				}
 			}
-
-			// if we found no matching buckets that we had access to we raise a 403
-			if (bucketsToSearch.Count == 0 && lackedAccessToAtLeastOneBucket)
-			{
-				return new ForbidResult();
-			}
-
-			// no buckets matched
-			if (bucketsToSearch.Count == 0)
+			catch (RemoteServiceException e)
 			{
-				// it's not an error to search were nothing matches, but it could be useful to know that it was due to the regular expression
-				return Ok(new SearchMultipleBucketsResult());
+				return Problem(e.Message, null, (int)HttpStatusCode.ServiceUnavailable);
 			}
 
 			SearchRequest searchRequest = searchBucketRequest;
@@ -583,134 +646,150 @@ namespace Jupiter.Controllers
 				return accessResultTarget;
 			}
 
-			BuildRecord? build = await _buildStore.GetBuildAsync(ns, bucket, buildId);
-			if (build == null)
+			try
 			{
-				return NotFound($"Build \"{build}\" in bucket {bucket} and namespace {ns}");
-			}
+				BuildRecord? build = await _buildStore.GetBuildAsync(ns, bucket, buildId);
+				if (build == null)
+				{
+					return NotFound($"Build \"{build}\" in bucket {bucket} and namespace {ns}");
+				}
 
-			CbObject originalBuildObject = build.BuildObject;
+				CbObject originalBuildObject = build.BuildObject;
 
-			CbWriter cbWriter = new CbWriter();
-			cbWriter.BeginObject();
-			List<string> fieldsToNotCopy = new List<string>();
-			if (!string.IsNullOrEmpty(request.NewBranch))
-			{
-				cbWriter.WriteString("branch", request.NewBranch);
-				fieldsToNotCopy.Add("branch");
-			}
-			if (!string.IsNullOrEmpty(request.NewProject))
-			{
-				cbWriter.WriteString("project", request.NewProject);
-				fieldsToNotCopy.Add("project");
-			}
-			if (!string.IsNullOrEmpty(request.NewPlatform))
-			{
-				cbWriter.WriteString("platform", request.NewPlatform);
-				fieldsToNotCopy.Add("platform");
-			}
-
-			foreach (CbField field in originalBuildObject)
-			{
-				string fieldName = field.Name.ToString();
-				if (fieldsToNotCopy.Contains(fieldName))
+				CbWriter cbWriter = new CbWriter();
+				cbWriter.BeginObject();
+				List<string> fieldsToNotCopy = new List<string>();
+				if (!string.IsNullOrEmpty(request.NewBranch))
 				{
-					continue;
+					cbWriter.WriteString("branch", request.NewBranch);
+					fieldsToNotCopy.Add("branch");
 				}
-				cbWriter.WriteField(field);
-			}
-			cbWriter.EndObject();
 
-			CbObject newBuildObject = cbWriter.ToObject();
-			uint ttl = (uint)_namespacePolicyResolver.GetPoliciesForNs(ns).DefaultTTL.TotalSeconds;
-			CbObjectId newBuildId = CbObjectId.NewObjectId();
-			await _buildStore.PutBuildAsync(targetNamespace, targetBucket, newBuildId, newBuildObject, ttl);
+				if (!string.IsNullOrEmpty(request.NewProject))
+				{
+					cbWriter.WriteString("project", request.NewProject);
+					fieldsToNotCopy.Add("project");
+				}
 
-			BlockContext newBlockContext = BlockContext.FromObject(newBuildObject);
+				if (!string.IsNullOrEmpty(request.NewPlatform))
+				{
+					cbWriter.WriteString("platform", request.NewPlatform);
+					fieldsToNotCopy.Add("platform");
+				}
 
-			await foreach ((string partName, CbObjectId partId)  in _buildStore.GetBuildPartsAsync(ns, bucket, buildId))
-			{
-				(RefRecord refRecord, BlobContents? contents)  = await _refService.GetAsync(ns, bucket, RefId.FromName($"{buildId}/{partId}"));
-				if (contents == null)
+				foreach (CbField field in originalBuildObject)
 				{
-					// large blob that is not inlined, needs to be copied
-					await _blobStore.CopyBlobAsync(ns, targetNamespace, refRecord.BlobIdentifier, bucket);
-					contents = await _blobStore.GetObjectAsync(targetNamespace, refRecord.BlobIdentifier);
+					string fieldName = field.Name.ToString();
+					if (fieldsToNotCopy.Contains(fieldName))
+					{
+						continue;
+					}
+
+					cbWriter.WriteField(field);
 				}
 
-				CbObject payloadObject = new CbObject(await contents.Stream.ReadAllBytesAsync());
-				BlobId payloadHash = refRecord.BlobIdentifier;
+				cbWriter.EndObject();
+
+				CbObject newBuildObject = cbWriter.ToObject();
+				uint ttl = (uint)_namespacePolicyResolver.GetPoliciesForNs(ns).DefaultTTL.TotalSeconds;
+				CbObjectId newBuildId = CbObjectId.NewObjectId();
+				await _buildStore.PutBuildAsync(targetNamespace, targetBucket, newBuildId, newBuildObject, ttl);
 
-				ConcurrentDictionary<ContentHash, List<BlobId>> contentIdToBlobs = new ConcurrentDictionary<ContentHash, List<BlobId>>();
+				BlockContext newBlockContext = BlockContext.FromObject(newBuildObject);
 
+				await foreach ((string partName, CbObjectId partId) in _buildStore.GetBuildPartsAsync(ns, bucket, buildId))
 				{
-					// check all attachments for content ids and blocks and copy them
-					IContentIdStore contentIdStore = HttpContext.RequestServices.GetService<IContentIdStore>()!;
+					(RefRecord refRecord, BlobContents? contents) = await _refService.GetAsync(ns, bucket, RefId.FromName($"{buildId}/{partId}"));
+					if (contents == null)
+					{
+						// large blob that is not inlined, needs to be copied
+						await _blobStore.CopyBlobAsync(ns, targetNamespace, refRecord.BlobIdentifier, bucket);
+						contents = await _blobStore.GetObjectAsync(targetNamespace, refRecord.BlobIdentifier);
+					}
+
+					CbObject payloadObject = new CbObject(await contents.Stream.ReadAllBytesAsync());
+					BlobId payloadHash = refRecord.BlobIdentifier;
 
-					using TelemetrySpan scope = _tracer.StartActiveSpan("copy.attachments").SetAttribute("operation.name", "copy.attachments");
+					ConcurrentDictionary<ContentHash, List<BlobId>> contentIdToBlobs = new ConcurrentDictionary<ContentHash, List<BlobId>>();
 
-					await foreach (Attachment attachment in _referenceResolver.GetAttachmentsAsync(ns, payloadObject, CancellationToken.None))
 					{
-						BlobId blobId = BlobId.FromIoHash(attachment.AsIoHash());
+						// check all attachments for content ids and blocks and copy them
+						IContentIdStore contentIdStore = HttpContext.RequestServices.GetService<IContentIdStore>()!;
 
-						// by default the raw hash is also the expected blob
-						List<BlobId> referencedBlobs = new() { blobId };
-						if (attachment is ContentIdAttachment contentIdAttachment)
+						using TelemetrySpan scope = _tracer.StartActiveSpan("copy.attachments").SetAttribute("operation.name", "copy.attachments");
+
+						await foreach (Attachment attachment in _referenceResolver.GetAttachmentsAsync(ns, payloadObject, CancellationToken.None))
 						{
-							// we are remapping the blobs so the raw hash is not valid
-							referencedBlobs.Clear();
+							BlobId blobId = BlobId.FromIoHash(attachment.AsIoHash());
 
-							IAsyncEnumerable<ContentIdMapping> mappings = contentIdStore.GetContentIdMappingsAsync(ns, contentIdAttachment.Identifier, CancellationToken.None);
-							await foreach (ContentIdMapping mapping in mappings)
+							// by default the raw hash is also the expected blob
+							List<BlobId> referencedBlobs = new() { blobId };
+							if (attachment is ContentIdAttachment contentIdAttachment)
 							{
-								// copy the content id
-								await contentIdStore.PutAsync(targetNamespace, contentIdAttachment.Identifier, mapping.ReferencedBlobs, mapping.Weight);
+								// we are remapping the blobs so the raw hash is not valid
+								referencedBlobs.Clear();
+
+								IAsyncEnumerable<ContentIdMapping> mappings =
+									contentIdStore.GetContentIdMappingsAsync(ns, contentIdAttachment.Identifier, CancellationToken.None);
+								await foreach (ContentIdMapping mapping in mappings)
+								{
+									// copy the content id
+									await contentIdStore.PutAsync(targetNamespace, contentIdAttachment.Identifier, mapping.ReferencedBlobs,
+										mapping.Weight);
 
-								referencedBlobs.AddRange(mapping.ReferencedBlobs);
+									referencedBlobs.AddRange(mapping.ReferencedBlobs);
+								}
 							}
-						}
 
-						BlobId? blockMetadata = await _blockStore.GetBlockMetadataAsync(ns, blobId);
-						if (blockMetadata != null)
-						{
-							await _blobStore.CopyBlobAsync(ns, targetNamespace, blockMetadata, bucket);
-							await _blockStore.PutBlockMetadataAsync(targetNamespace, blobId, blockMetadata);
-						}
+							BlobId? blockMetadata = await _blockStore.GetBlockMetadataAsync(ns, blobId);
+							if (blockMetadata != null)
+							{
+								await _blobStore.CopyBlobAsync(ns, targetNamespace, blockMetadata, bucket);
+								await _blockStore.PutBlockMetadataAsync(targetNamespace, blobId, blockMetadata);
+							}
 
-						contentIdToBlobs[blobId] = referencedBlobs;
+							contentIdToBlobs[blobId] = referencedBlobs;
+						}
 					}
-				}
-				
-				await _buildStore.PutBuildPartAsync(targetNamespace, targetBucket, newBuildId, partId, partName, ttl);
 
-				List<ContentHash> missingHashes = new List<ContentHash>();
-				const int RetryAttempts = 3;
-				for (int i = 0; i < RetryAttempts; i++)
-				{
-					missingHashes = await PutBuildPartRefAsync(targetNamespace, targetBucket, newBuildId, partId, newBlockContext, payloadObject, payloadHash);
-					if (missingHashes.Count != 0)
-					{
-						using TelemetrySpan scope = _tracer.StartActiveSpan("copy.blobs").SetAttribute("operation.name", "copy.blobs");
+					await _buildStore.PutBuildPartAsync(targetNamespace, targetBucket, newBuildId, partId, partName, ttl);
 
-						// copy missing blobs and try to finalize again
-						foreach (ContentHash missingHash in missingHashes)
+					List<ContentHash> missingHashes = new List<ContentHash>();
+					const int RetryAttempts = 3;
+					for (int i = 0; i < RetryAttempts; i++)
+					{
+						missingHashes = await PutBuildPartRefAsync(targetNamespace, targetBucket, newBuildId, partId, newBlockContext, payloadObject,
+							payloadHash);
+						if (missingHashes.Count != 0)
 						{
-							if (contentIdToBlobs.TryGetValue(missingHash, out List<BlobId>? blobs))
+							using TelemetrySpan scope = _tracer.StartActiveSpan("copy.blobs").SetAttribute("operation.name", "copy.blobs");
+
+							// copy missing blobs and try to finalize again
+							foreach (ContentHash missingHash in missingHashes)
 							{
-								foreach (BlobId blobId in blobs)
+								if (contentIdToBlobs.TryGetValue(missingHash, out List<BlobId>? blobs))
 								{
-									await _blobStore.CopyBlobAsync(ns, targetNamespace, blobId, bucket);
+									foreach (BlobId blobId in blobs)
+									{
+										await _blobStore.CopyBlobAsync(ns, targetNamespace, blobId, bucket);
+									}
 								}
 							}
 						}
 					}
+
+					if (missingHashes.Count != 0)
+					{
+						throw new Exception($"Hashes missing when copying build after {RetryAttempts} attempts.");
+					}
 				}
-				if (missingHashes.Count != 0)
-				{
-					throw new Exception($"Hashes missing when copying build after {RetryAttempts} attempts.");
-				}
+
+				return Ok(new CopyBuildResponse {BuildId = newBuildId});
+			}
+			catch (RemoteServiceException e)
+			{
+				return Problem(e.Message, null, (int)HttpStatusCode.ServiceUnavailable);
 			}
-			return Ok(new CopyBuildResponse {BuildId = newBuildId});
 		}
 		#endregion
 
diff --git a/Engine/Source/Programs/UnrealCloudDDC/Jupiter/Controllers/CompressedBlobController.cs b/Engine/Source/Programs/UnrealCloudDDC/Jupiter/Controllers/CompressedBlobController.cs
--- a/Engine/Source/Programs/UnrealCloudDDC/Jupiter/Controllers/CompressedBlobController.cs
+++ b/Engine/Source/Programs/UnrealCloudDDC/Jupiter/Controllers/CompressedBlobController.cs
@@ -261,6 +261,10 @@ namespace Jupiter.Controllers
 			{
 				return Problem(e.Message, null, (int)HttpStatusCode.RequestTimeout);
 			}
+			catch (RemoteServiceException e)
+			{
+				return Problem(e.Message, null, (int)HttpStatusCode.ServiceUnavailable);
+			}
 		}
 
 		[HttpPost("{ns}")]
@@ -301,6 +305,10 @@ namespace Jupiter.Controllers
 			{
 				return Problem(e.Message, null, (int)HttpStatusCode.RequestTimeout);
 			}
+			catch (RemoteServiceException e)
+			{
+				return Problem(e.Message, null, (int)HttpStatusCode.ServiceUnavailable);
+			}
 		}
 
 		[HttpDelete("{ns}/{id}")]
@@ -331,6 +339,10 @@ namespace Jupiter.Controllers
 				{
 					// if the blob is already missing that is fine
 				}
+				catch (RemoteServiceException e)
+				{
+					return Problem(e.Message, null, (int)HttpStatusCode.ServiceUnavailable);
+				}
 			}
 
 			// TODO: we should delete the cid from the content id store as well, but it has no delete operations yet
diff --git a/Engine/Source/Programs/UnrealCloudDDC/Jupiter/Controllers/ContentIdController.cs b/Engine/Source/Programs/UnrealCloudDDC/Jupiter/Controllers/ContentIdController.cs
--- a/Engine/Source/Programs/UnrealCloudDDC/Jupiter/Controllers/ContentIdController.cs
+++ b/Engine/Source/Programs/UnrealCloudDDC/Jupiter/Controllers/ContentIdController.cs
@@ -1,9 +1,11 @@
 // Copyright Epic Games, Inc. All Rights Reserved.
 
+using System.Net;
 using System.Net.Mime;
 using System.Threading.Tasks;
 using EpicGames.AspNet;
 using EpicGames.Horde.Storage;
+using Jupiter.Common;
 using Jupiter.Implementation;
 using Microsoft.AspNetCore.Authorization;
 using Microsoft.AspNetCore.Mvc;
@@ -42,14 +44,22 @@ namespace Jupiter.Controllers
 			{
 				return result;
 			}
-			BlobId[]? blobs = await _contentIdStore.ResolveAsync(ns, contentId, mustBeContentId: true);
 
-			if (blobs == null)
+			try
 			{
-				return NotFound(new ProblemDetails { Title = $"Unable to resolve content id {contentId} ({ns})." });
-			}
+				BlobId[]? blobs = await _contentIdStore.ResolveAsync(ns, contentId, mustBeContentId: true);
+
+				if (blobs == null)
+				{
+					return NotFound(new ProblemDetails { Title = $"Unable to resolve content id {contentId} ({ns})." });
+				}
 
-			return Ok(new ResolvedContentIdResponse(blobs));
+				return Ok(new ResolvedContentIdResponse(blobs));
+			}
+			catch (RemoteServiceException e)
+			{
+				return Problem(e.Message, null, (int)HttpStatusCode.ServiceUnavailable);
+			}
 		}
 
 		/// <summary>
@@ -69,7 +79,15 @@ namespace Jupiter.Controllers
 			{
 				return result;
 			}
-			await _contentIdStore.PutAsync(ns, contentId, new BlobId[] {blobIdentifier}, contentWeight);
+
+			try
+			{
+				await _contentIdStore.PutAsync(ns, contentId, new BlobId[] {blobIdentifier}, contentWeight);
+			}
+			catch (RemoteServiceException e)
+			{
+				return Problem(e.Message, null, (int)HttpStatusCode.ServiceUnavailable);
+			}
 
 			return Ok();
 		}
diff --git a/Engine/Source/Programs/UnrealCloudDDC/Jupiter/Controllers/ObjectController.cs b/Engine/Source/Programs/UnrealCloudDDC/Jupiter/Controllers/ObjectController.cs
--- a/Engine/Source/Programs/UnrealCloudDDC/Jupiter/Controllers/ObjectController.cs
+++ b/Engine/Source/Programs/UnrealCloudDDC/Jupiter/Controllers/ObjectController.cs
@@ -11,6 +11,7 @@ using System.Threading.Tasks;
 using EpicGames.AspNet;
 using EpicGames.Horde.Storage;
 using EpicGames.Serialization;
+using Jupiter.Common;
 using Jupiter.Common.Implementation;
 using Jupiter.Implementation;
 using Jupiter.Utils;
@@ -69,6 +70,10 @@ namespace Jupiter.Controllers
 			{
 				return NotFound(new ValidationProblemDetails { Title = $"Object {e.Blob} not found" });
 			}
+			catch (RemoteServiceException e)
+			{
+				return Problem(e.Message, null, (int)HttpStatusCode.ServiceUnavailable);
+			}
 		}
 
 		[HttpHead("{ns}/{id}")]
@@ -83,11 +88,18 @@ namespace Jupiter.Controllers
 				return result;
 			}
 
-			bool exists = await _storage.ExistsAsync(ns, id);
+			try
+			{
+				bool exists = await _storage.ExistsAsync(ns, id);
 
-			if (!exists)
+				if (!exists)
+				{
+					return NotFound(new ValidationProblemDetails { Title = $"Object {id} not found" });
+				}
+			}
+			catch (RemoteServiceException e)
 			{
-				return NotFound(new ValidationProblemDetails { Title = $"Object {id} not found" });
+				return Problem(e.Message, null, (int)HttpStatusCode.ServiceUnavailable);
 			}
 
 			return Ok();
@@ -169,6 +181,10 @@ namespace Jupiter.Controllers
 			{
 				return Problem(e.Message, null, (int)HttpStatusCode.RequestTimeout);
 			}
+			catch (RemoteServiceException e)
+			{
+				return Problem(e.Message, null, (int)HttpStatusCode.ServiceUnavailable);
+			}
 		}
 
 		[HttpGet("{ns}/{id}/references")]
@@ -191,6 +207,10 @@ namespace Jupiter.Controllers
 			{
 				return NotFound(new ValidationProblemDetails { Title = $"Object {e.Blob} not found" });
 			}
+			catch (RemoteServiceException e)
+			{
+				return Problem(e.Message, null, (int)HttpStatusCode.ServiceUnavailable);
+			}
 
 			byte[] blobContents = await blob.Stream.ToByteArrayAsync(HttpContext.RequestAborted);
 			if (blobContents.Length == 0)
@@ -221,6 +241,10 @@ namespace Jupiter.Controllers
 			{
 				return BadRequest(new ValidationProblemDetails { Title = $"Object {id} is missing blobs", Detail = $"Following blobs are missing: {string.Join(",", e.MissingBlobs)}" });
 			}
+			catch (RemoteServiceException e)
+			{
+				return Problem(e.Message, null, (int)HttpStatusCode.ServiceUnavailable);
+			}
 		}
 
 		[HttpDelete("{ns}/{id}")]
@@ -234,12 +258,19 @@ namespace Jupiter.Controllers
 				return result;
 			}
 
-			await _storage.DeleteObjectAsync(ns, id, HttpContext.RequestAborted);
+			try
+			{
+				await _storage.DeleteObjectAsync(ns, id, HttpContext.RequestAborted);
 
-			return Ok(new DeletedResponse
+				return Ok(new DeletedResponse
+				{
+					DeletedCount = 1
+				});
+			}
+			catch (RemoteServiceException e)
 			{
-				DeletedCount = 1
-			});
+				return Problem(e.Message, null, (int)HttpStatusCode.ServiceUnavailable);
+			}
 		}
 
 		[HttpDelete("{ns}")]
@@ -252,7 +283,14 @@ namespace Jupiter.Controllers
 				return result;
 			}
 
-			await _storage.DeleteNamespaceAsync(ns, HttpContext.RequestAborted);
+			try
+			{
+				await _storage.DeleteNamespaceAsync(ns, HttpContext.RequestAborted);
+			}
+			catch (RemoteServiceException e)
+			{
+				return Problem(e.Message, null, (int)HttpStatusCode.ServiceUnavailable);
+			}
 
 			return Ok();
 		}
diff --git a/Engine/Source/Programs/UnrealCloudDDC/Jupiter/Controllers/ReferencesController.cs b/Engine/Source/Programs/UnrealCloudDDC/Jupiter/Controllers/ReferencesController.cs
--- a/Engine/Source/Programs/UnrealCloudDDC/Jupiter/Controllers/ReferencesController.cs
+++ b/Engine/Source/Programs/UnrealCloudDDC/Jupiter/Controllers/ReferencesController.cs
@@ -89,25 +89,32 @@ namespace Jupiter.Controllers
 		[ProducesResponseType(type: typeof(ProblemDetails), 400)]
 		public async Task<IActionResult> GetNamespacesAsync()
 		{
-			NamespaceId[] namespaces = await _refService.GetNamespacesAsync(HttpContext.RequestAborted).ToArrayAsync(HttpContext.RequestAborted);
-
-			// filter namespaces down to only the namespaces the user has access to
-			List<NamespaceId> namespacesWithAccess = new();
-			foreach (NamespaceId ns in namespaces)
+			try
 			{
-				ActionResult? accessResult = await _requestHelper.HasAccessToNamespaceAsync(User, Request, ns, new[] { JupiterAclAction.ReadObject });
-				if (accessResult == null)
+				NamespaceId[] namespaces = await _refService.GetNamespacesAsync(HttpContext.RequestAborted).ToArrayAsync(HttpContext.RequestAborted);
+
+				// filter namespaces down to only the namespaces the user has access to
+				List<NamespaceId> namespacesWithAccess = new();
+				foreach (NamespaceId ns in namespaces)
 				{
-					namespacesWithAccess.Add(ns);
+					ActionResult? accessResult = await _requestHelper.HasAccessToNamespaceAsync(User, Request, ns, new[] { JupiterAclAction.ReadObject });
+					if (accessResult == null)
+					{
+						namespacesWithAccess.Add(ns);
+					}
+				}
+
+				if (!namespacesWithAccess.Any())
+				{
+					return new ForbidResult();
 				}
-			}
 
-			if (!namespacesWithAccess.Any())
+				return Ok(new GetNamespacesResponse(namespacesWithAccess.ToArray()));
+			}
+			catch (RemoteServiceException e)
 			{
-				return new ForbidResult();
+				return Problem(e.Message, null, (int)HttpStatusCode.ServiceUnavailable);
 			}
-
-			return Ok(new GetNamespacesResponse(namespacesWithAccess.ToArray()));
 		}
 
 		/// <summary>
@@ -132,8 +139,16 @@ namespace Jupiter.Controllers
 			{
 				return Forbid();
 			}
-			RefId[] refIds = await _refService.GetRecordsInBucketAsync(ns, bucket).ToArrayAsync();
-			return Ok(new EnumerateBucketResponse(refIds));
+
+			try
+			{
+				RefId[] refIds = await _refService.GetRecordsInBucketAsync(ns, bucket).ToArrayAsync();
+				return Ok(new EnumerateBucketResponse(refIds));
+			}
+			catch (RemoteServiceException e)
+			{
+				return Problem(e.Message, null, (int)HttpStatusCode.ServiceUnavailable);
+			}
 		}
 
 		/// <summary>
@@ -477,6 +492,10 @@ namespace Jupiter.Controllers
 			{
 				return NotFound(new ProblemDetails { Title = $"Object {e.Blob} in {e.Ns} not found" });
 			}
+			catch (RemoteServiceException e)
+			{
+				return Problem(e.Message, null, (int)HttpStatusCode.ServiceUnavailable);
+			}
 		}
 
 		/// <summary>
@@ -517,6 +536,10 @@ namespace Jupiter.Controllers
 			{
 				return NotFound(new ProblemDetails { Title = $"Object {e.Blob} in {e.Ns} not found" });
 			}
+			catch (RemoteServiceException e)
+			{
+				return Problem(e.Message, null, (int)HttpStatusCode.ServiceUnavailable);
+			}
 		}
 
 		/// <summary>
@@ -572,6 +595,10 @@ namespace Jupiter.Controllers
 			{
 				return NotFound(new ProblemDetails { Title = $"Object {e.Blob} in {e.Ns} not found" });
 			}
+			catch (RemoteServiceException e)
+			{
+				return Problem(e.Message, null, (int)HttpStatusCode.ServiceUnavailable);
+			}
 		}
 
 		/// <summary>
@@ -610,6 +637,10 @@ namespace Jupiter.Controllers
 			{
 				return NotFound(new ProblemDetails { Title = $"Object {e.Bucket} {e.Key} did not exist" });
 			}
+			catch (RemoteServiceException e)
+			{
+				return Problem(e.Message, null, (int)HttpStatusCode.ServiceUnavailable);
+			}
 
 			if (blob == null)
 			{
@@ -651,6 +682,10 @@ namespace Jupiter.Controllers
 			{
 				return BadRequest(new ValidationProblemDetails { Title = $"Object {bucket} {key} did not exist", Detail = $"Following blobs are missing: {string.Join(",", e.MissingBlobs)}" });
 			}
+			catch (RemoteServiceException e)
+			{
+				return Problem(e.Message, null, (int)HttpStatusCode.ServiceUnavailable);
+			}
 		}
 
 		/// <summary>
@@ -722,6 +757,10 @@ namespace Jupiter.Controllers
 			{
 				return NotFound(new ProblemDetails { Title = $"Object {bucket} {key} in namespace {ns} was missing some blobs" });
 			}
+			catch (RemoteServiceException e)
+			{
+				return Problem(e.Message, null, (int)HttpStatusCode.ServiceUnavailable);
+			}
 
 			return Ok();
 		}
@@ -902,6 +941,10 @@ namespace Jupiter.Controllers
 			{
 				return Problem(e.Message, null, (int)HttpStatusCode.RequestTimeout);
 			}
+			catch (RemoteServiceException e)
+			{
+				return Problem(e.Message, null, (int)HttpStatusCode.ServiceUnavailable);
+			}
 
 			try
 			{
@@ -942,6 +985,10 @@ namespace Jupiter.Controllers
 				await Response.Body.WriteAsync(payload);
 				return new EmptyResult();
 			}
+			catch (RemoteServiceException e)
+			{
+				return Problem(e.Message, null, (int)HttpStatusCode.ServiceUnavailable);
+			}
 		}
 
 		[HttpPut("{ns}/{bucket}/{key}", Order = 300)]
@@ -996,15 +1043,26 @@ namespace Jupiter.Controllers
 					Title = $"Incorrect hash, got hash \"{e.SuppliedHash}\" but hash of content was determined to be \"{e.ContentHash}\""
 				});
 			}
+			catch (RemoteServiceException e)
+			{
+				return Problem(e.Message, null, (int)HttpStatusCode.ServiceUnavailable);
+			}
 
 			CbObject rootObject = packageReader.RootObject;
 			BlobId rootObjectHash = BlobId.FromIoHash(packageReader.RootHash);
 
-			(ContentId[] missingReferences, BlobId[] missingBlobs) = await _refService.PutAsync(ns, bucket, key, rootObjectHash, rootObject, cancellationToken: HttpContext.RequestAborted);
+			try
+			{
+				(ContentId[] missingReferences, BlobId[] missingBlobs) = await _refService.PutAsync(ns, bucket, key, rootObjectHash, rootObject, cancellationToken: HttpContext.RequestAborted);
 
-			List<ContentHash> missingHashes = new List<ContentHash>(missingReferences);
-			missingHashes.AddRange(missingBlobs);
-			return Ok(new PutObjectResponse(missingHashes.ToArray()));
+				List<ContentHash> missingHashes = new List<ContentHash>(missingReferences);
+				missingHashes.AddRange(missingBlobs);
+				return Ok(new PutObjectResponse(missingHashes.ToArray()));
+			}
+			catch (RemoteServiceException e)
+			{
+				return Problem(e.Message, null, (int)HttpStatusCode.ServiceUnavailable);
+			}
 		}
 
 		[HttpPost("{ns}/{bucket}/{key}/finalize/{hash}.{format?}")]
@@ -1039,6 +1097,10 @@ namespace Jupiter.Controllers
 			{
 				return NotFound(e.Message);
 			}
+			catch (RemoteServiceException e)
+			{
+				return Problem(e.Message, null, (int)HttpStatusCode.ServiceUnavailable);
+			}
 		}
 
 		[HttpGet("{ns}/{bucket}/{key}/blobs/{id}")]
@@ -1097,6 +1159,10 @@ namespace Jupiter.Controllers
 			{
 				return NotFound(new ValidationProblemDetails { Title = $"Content Id {e.ContentId} not found" });
 			}
+			catch (RemoteServiceException e)
+			{
+				return Problem(e.Message, null, (int)HttpStatusCode.ServiceUnavailable);
+			}
 		}
 
 		[HttpPut("{ns}/{bucket}/{key}/blobs/{id}")]
@@ -1164,6 +1230,10 @@ namespace Jupiter.Controllers
 			{
 				return Problem(e.Message, null, (int)HttpStatusCode.RequestTimeout);
 			}
+			catch (RemoteServiceException e)
+			{
+				return Problem(e.Message, null, (int)HttpStatusCode.ServiceUnavailable);
+			}
 		}
 
 		[HttpPost("{ns}")]
@@ -1415,6 +1485,10 @@ namespace Jupiter.Controllers
 			{
 				return NotFound(new ProblemDetails { Title = $"Namespace {e.Namespace} did not exist" });
 			}
+			catch (RemoteServiceException e)
+			{
+				return Problem(e.Message, null, (int)HttpStatusCode.ServiceUnavailable);
+			}
 
 			return NoContent();
 		}
@@ -1445,6 +1519,10 @@ namespace Jupiter.Controllers
 			{
 				return NotFound(new ProblemDetails { Title = $"Namespace {e.Namespace} did not exist" });
 			}
+			catch (RemoteServiceException e)
+			{
+				return Problem(e.Message, null, (int)HttpStatusCode.ServiceUnavailable);
+			}
 
 			return Ok(new BucketDeletedResponse(countOfDeletedRecords));
 		}
@@ -1482,6 +1560,10 @@ namespace Jupiter.Controllers
 			{
 				return Ok(new RefDeletedResponse(0));
 			}
+			catch (RemoteServiceException e)
+			{
+				return Problem(e.Message, null, (int)HttpStatusCode.ServiceUnavailable);
+			}
 		}
 	}
 
diff --git a/Engine/Source/Programs/UnrealCloudDDC/Jupiter/Controllers/ReplicationLogController.cs b/Engine/Source/Programs/UnrealCloudDDC/Jupiter/Controllers/ReplicationLogController.cs
--- a/Engine/Source/Programs/UnrealCloudDDC/Jupiter/Controllers/ReplicationLogController.cs
+++ b/Engine/Source/Programs/UnrealCloudDDC/Jupiter/Controllers/ReplicationLogController.cs
@@ -4,6 +4,7 @@ using System;
 using System.Collections.Generic;
 using System.ComponentModel.DataAnnotations;
 using System.Linq;
+using System.Net;
 using System.Text.Json.Serialization;
 using System.Threading;
 using System.Threading.Tasks;
@@ -51,7 +52,14 @@ namespace Jupiter.Controllers
 				return result;
 			}
 
-			return Ok(new ReplicationLogSnapshots(await _replicationLog.GetSnapshotsAsync(ns).ToListAsync()));
+			try
+			{
+				return Ok(new ReplicationLogSnapshots(await _replicationLog.GetSnapshotsAsync(ns).ToListAsync()));
+			}
+			catch (RemoteServiceException e)
+			{
+				return Problem(e.Message, null, (int)HttpStatusCode.ServiceUnavailable);
+			}
 		}
 
 		[HttpPost("snapshots/{ns}/create")]
@@ -67,9 +75,16 @@ namespace Jupiter.Controllers
 				return result;
 			}
 
-			ReplicationLogSnapshotBuilder builder = ActivatorUtilities.CreateInstance<ReplicationLogSnapshotBuilder>(_provider);
-			BlobId snapshotBlob = await builder.BuildSnapshotAsync(ns, _snapshotSettings.CurrentValue.SnapshotStorageNamespace, CancellationToken.None);
-			return Ok(new SnapshotCreatedResponse(snapshotBlob));
+			try
+			{
+				ReplicationLogSnapshotBuilder builder = ActivatorUtilities.CreateInstance<ReplicationLogSnapshotBuilder>(_provider);
+				BlobId snapshotBlob = await builder.BuildSnapshotAsync(ns, _snapshotSettings.CurrentValue.SnapshotStorageNamespace, CancellationToken.None);
+				return Ok(new SnapshotCreatedResponse(snapshotBlob));
+			}
+			catch (RemoteServiceException e)
+			{
+				return Problem(e.Message, null, (int)HttpStatusCode.ServiceUnavailable);
+			}
 		}
 
 		[HttpGet("incremental/{ns}")]
@@ -133,6 +148,10 @@ namespace Jupiter.Controllers
 					Title = $"Namespace {ns} was not found",
 				});
 			}
+			catch (RemoteServiceException e)
+			{
+				return Problem(e.Message, null, (int)HttpStatusCode.ServiceUnavailable);
+			}
 		}
 
 		[HttpGet("blobs/{ns}/{replicationBucket}")]
@@ -185,6 +204,10 @@ namespace Jupiter.Controllers
 					Title = $"Namespace {ns} was not found",
 				});
 			}
+			catch (RemoteServiceException e)
+			{
+				return Problem(e.Message, null, (int)HttpStatusCode.ServiceUnavailable);
+			}
 		}
 	}
 
diff --git a/Engine/Source/Programs/UnrealCloudDDC/Jupiter/Controllers/SymbolsController.cs b/Engine/Source/Programs/UnrealCloudDDC/Jupiter/Controllers/SymbolsController.cs
--- a/Engine/Source/Programs/UnrealCloudDDC/Jupiter/Controllers/SymbolsController.cs
+++ b/Engine/Source/Programs/UnrealCloudDDC/Jupiter/Controllers/SymbolsController.cs
@@ -5,6 +5,7 @@ using System.Collections.Generic;
 using System.ComponentModel.DataAnnotations;
 using System.IO;
 using System.Linq;
+using System.Net;
 using System.Net.Http;
 using System.Net.Mime;
 using System.Text;
@@ -97,6 +98,10 @@ namespace Jupiter.Controllers
 			{
 				return NotFound("Symbol not found");
 			}
+			catch (RemoteServiceException e)
+			{
+				return Problem(e.Message, null, (int)HttpStatusCode.ServiceUnavailable);
+			}
 
 			if (refContents == null)
 			{
@@ -108,67 +113,78 @@ namespace Jupiter.Controllers
 			CbObject cb = new CbObject(blobMemory);
 			IoHash payloadHash = cb["pdbPayload"].AsBinaryAttachment().Hash;
 
-			BlobContents referencedBlobContents = await _blobStore.GetObjectAsync(ns, BlobId.FromIoHash(payloadHash), storageLayers: null, supportsRedirectUri: supportsRedirect, bucketHint: bucket);
-
-			if (referencedBlobContents.RedirectUri != null)
-			{
-				return Redirect(referencedBlobContents.RedirectUri.ToString());
-			}
-
-			if (_nginxRedirectHelper.CanRedirect(Request, referencedBlobContents))
+			try
 			{
-				return _nginxRedirectHelper.CreateActionResult(referencedBlobContents, CustomMediaTypeNames.UnrealCompressedBuffer);
-			}
+				BlobContents referencedBlobContents = await _blobStore.GetObjectAsync(ns, BlobId.FromIoHash(payloadHash), storageLayers: null,
+					supportsRedirectUri: supportsRedirect, bucketHint: bucket);
 
-			async Task WriteBody(BlobContents blobContents, string contentType)
-			{
-				IServerTiming? serverTiming = Request.HttpContext.RequestServices.GetService<IServerTiming>();
-				using ServerTimingMetricScoped? serverTimingScope = serverTiming?.CreateServerTimingMetricScope("body.write", "Time spent writing body");
-
-				long contentLength = blobContents.Length;
-				using TelemetrySpan scope = _tracer.StartActiveSpan("body.write").SetAttribute("operation.name", "body.write");
-				scope.SetAttribute("content-length", contentLength);
-				const int BufferSize = 64 * 1024;
-				Stream outputStream = Response.Body;
-				Response.ContentLength = contentLength;
-				Response.ContentType = contentType;
-				Response.StatusCode = StatusCodes.Status200OK;
-				try
+				if (referencedBlobContents.RedirectUri != null)
 				{
-					await StreamCopyOperation.CopyToAsync(blobContents.Stream, outputStream, count: null, bufferSize: BufferSize, cancel: Response.HttpContext.RequestAborted);
+					return Redirect(referencedBlobContents.RedirectUri.ToString());
 				}
-				catch (OperationCanceledException)
+
+				if (_nginxRedirectHelper.CanRedirect(Request, referencedBlobContents))
 				{
-					// do not raise exceptions for cancelled writes
-					// as we have already started writing a response we can not change the status code
-					// so we just drop a warning and proceed
-					_logger.LogWarning("The operation was canceled while writing the body");
+					return _nginxRedirectHelper.CreateActionResult(referencedBlobContents, CustomMediaTypeNames.UnrealCompressedBuffer);
 				}
-			}
 
-			string responseType = _formatResolver.GetResponseType(Request, null, MediaTypeNames.Application.Octet);
-			Tracer.CurrentSpan.SetAttribute("response-type", responseType);
-			switch (responseType)
-			{
-				case CustomMediaTypeNames.UnrealCompressedBuffer:
+				async Task WriteBody(BlobContents blobContents, string contentType)
+				{
+					IServerTiming? serverTiming = Request.HttpContext.RequestServices.GetService<IServerTiming>();
+					using ServerTimingMetricScoped? serverTimingScope =
+						serverTiming?.CreateServerTimingMetricScope("body.write", "Time spent writing body");
+
+					long contentLength = blobContents.Length;
+					using TelemetrySpan scope = _tracer.StartActiveSpan("body.write").SetAttribute("operation.name", "body.write");
+					scope.SetAttribute("content-length", contentLength);
+					const int BufferSize = 64 * 1024;
+					Stream outputStream = Response.Body;
+					Response.ContentLength = contentLength;
+					Response.ContentType = contentType;
+					Response.StatusCode = StatusCodes.Status200OK;
+					try
 					{
-						await WriteBody(referencedBlobContents, CustomMediaTypeNames.UnrealCompressedBuffer);
-						break;
+						await StreamCopyOperation.CopyToAsync(blobContents.Stream, outputStream, count: null, bufferSize: BufferSize,
+							cancel: Response.HttpContext.RequestAborted);
 					}
-				case MediaTypeNames.Application.Octet:
+					catch (OperationCanceledException)
 					{
-						CompressedBufferUtils utils = new CompressedBufferUtils(_tracer, _bufferedPayloadFactory);
-						(IBufferedPayload, IoHash) decompressResult = await utils.DecompressContentAsync(referencedBlobContents.Stream, (ulong)referencedBlobContents.Length, HttpContext.RequestAborted);
-						using IBufferedPayload payload = decompressResult.Item1;
-
-						await using Stream s = payload.GetStream();
-						await using BlobContents contents = new BlobContents(s, s.Length);
-						await WriteBody(contents, MediaTypeNames.Application.Octet);
-						break;
+						// do not raise exceptions for cancelled writes
+						// as we have already started writing a response we can not change the status code
+						// so we just drop a warning and proceed
+						_logger.LogWarning("The operation was canceled while writing the body");
 					}
+				}
 
-				default:
-					throw new NotImplementedException($"Unknown expected response type {responseType}");
+				string responseType = _formatResolver.GetResponseType(Request, null, MediaTypeNames.Application.Octet);
+				Tracer.CurrentSpan.SetAttribute("response-type", responseType);
+				switch (responseType)
+				{
+					case CustomMediaTypeNames.UnrealCompressedBuffer:
+						{
+							await WriteBody(referencedBlobContents, CustomMediaTypeNames.UnrealCompressedBuffer);
+							break;
+						}
+					case MediaTypeNames.Application.Octet:
+						{
+							CompressedBufferUtils utils = new CompressedBufferUtils(_tracer, _bufferedPayloadFactory);
+							(IBufferedPayload, IoHash) decompressResult = await utils.DecompressContentAsync(referencedBlobContents.Stream,
+								(ulong)referencedBlobContents.Length, HttpContext.RequestAborted);
+							using IBufferedPayload payload = decompressResult.Item1;
+
+							await using Stream s = payload.GetStream();
+							await using BlobContents contents = new BlobContents(s, s.Length);
+							await WriteBody(contents, MediaTypeNames.Application.Octet);
+							break;
+						}
+
+					default:
+						throw new NotImplementedException($"Unknown expected response type {responseType}");
+				}
+			}
+			catch (RemoteServiceException e)
+			{
+				return Problem(e.Message, null, (int)HttpStatusCode.ServiceUnavailable);
 			}
 
 			// content was written above
@@ -230,14 +246,21 @@ namespace Jupiter.Controllers
 
 			bool? bypassCache = _namespacePolicyResolver.GetPoliciesForNs(ns).BypassCacheOnWrite;
 
-			await _blobStore.PutObjectKnownHashAsync(ns, payload, attachmentHash, bucketName, bypassCache: bypassCache, cancellationToken: HttpContext.RequestAborted);
+			try
+			{
+				await _blobStore.PutObjectKnownHashAsync(ns, payload, attachmentHash, bucketName, bypassCache: bypassCache, cancellationToken: HttpContext.RequestAborted);
 
-			(ContentId[], BlobId[]) missingHashes = await _refService.PutAsync(ns, bucketName, RefId.FromName($"{moduleName}.{pdbIdentifier}{pdbAge}.{filename}"), blobHeader, o, cancellationToken: HttpContext.RequestAborted);
+				(ContentId[], BlobId[]) missingHashes = await _refService.PutAsync(ns, bucketName, RefId.FromName($"{moduleName}.{pdbIdentifier}{pdbAge}.{filename}"), blobHeader, o, cancellationToken: HttpContext.RequestAborted);
 
-			if (missingHashes.Item1.Any() || missingHashes.Item2.Any())
+				if (missingHashes.Item1.Any() || missingHashes.Item2.Any())
+				{
+					// this shouldn't be able to handle as we upload the object just before adding the ref
+					return BadRequest("Missing references from object");
+				}
+			}
+			catch (RemoteServiceException e)
 			{
-				// this shouldn't be able to handle as we upload the object just before adding the ref
-				return BadRequest("Missing references from object");
+				return Problem(e.Message, null, (int)HttpStatusCode.ServiceUnavailable);
 			}
 
 			return Ok(new PutSymbolResponse(moduleName, pdbIdentifier, pdbAge, attachmentIoHash));
diff --git a/Engine/Source/Programs/UnrealCloudDDC/Jupiter/Implementation/Blob/ScyllaBlobIndex.cs b/Engine/Source/Programs/UnrealCloudDDC/Jupiter/Implementation/Blob/ScyllaBlobIndex.cs
--- a/Engine/Source/Programs/UnrealCloudDDC/Jupiter/Implementation/Blob/ScyllaBlobIndex.cs
+++ b/Engine/Source/Programs/UnrealCloudDDC/Jupiter/Implementation/Blob/ScyllaBlobIndex.cs
@@ -106,24 +106,38 @@ public class ScyllaBlobIndex : IBlobIndex
 		region ??= _jupiterSettings.CurrentValue.CurrentSite;
 		using TelemetrySpan scope = _tracer.BuildScyllaSpan("scylla.insert_blob_index").SetAttribute("resource.name", $"{ns}.{id}");
 
-		await _mapper.InsertAsync<ScyllaBlobIndexEntry>(new ScyllaBlobIndexEntry(ns.ToString(), id, region));
+		try
+		{
+			await _mapper.InsertAsync<ScyllaBlobIndexEntry>(new ScyllaBlobIndexEntry(ns.ToString(), id, region));
+		}
+		catch (QueryExecutionException e)
+		{
+			throw new RemoteServiceException("Error communicating with Scylla", e);
+		}
 	}
 
 	private async Task<List<string>?> GetOldBlobRegionsAsync(NamespaceId ns, BlobId id)
 	{
 		using TelemetrySpan scope = _tracer.BuildScyllaSpan("scylla.fetch_old_blob_index").SetAttribute("resource.name", $"{ns}.{id}");
 
-		if (_scyllaSessionManager.IsScylla)
+		try
 		{
-			ScyllaBlobIndexTable o;
-			o = await _mapper.SingleOrDefaultAsync<ScyllaBlobIndexTable>("SELECT namespace, blob_id, regions, references FROM blob_index WHERE namespace = ? AND blob_id = ?", ns.ToString(), new ScyllaBlobIdentifier(id));
-			return o?.Regions?.ToList();
+			if (_scyllaSessionManager.IsScylla)
+			{
+				ScyllaBlobIndexTable o;
+				o = await _mapper.SingleOrDefaultAsync<ScyllaBlobIndexTable>("SELECT namespace, blob_id, regions, references FROM blob_index WHERE namespace = ? AND blob_id = ?", ns.ToString(), new ScyllaBlobIdentifier(id));
+				return o?.Regions?.ToList();
+			}
+			else
+			{
+				CassandraBlobIndexTable o;
+				o = await _mapper.SingleOrDefaultAsync<CassandraBlobIndexTable>("SELECT namespace, blob_id, regions, references FROM blob_index WHERE namespace = ? AND blob_id = ?", ns.ToString(), id.HashData);
+				return o?.Regions?.ToList();
+			}
 		}
-		else
+		catch (QueryExecutionException e)
 		{
-			CassandraBlobIndexTable o;
-			o = await _mapper.SingleOrDefaultAsync<CassandraBlobIndexTable>("SELECT namespace, blob_id, regions, references FROM blob_index WHERE namespace = ? AND blob_id = ?", ns.ToString(), id.HashData);
-			return o?.Regions?.ToList();
+			throw new RemoteServiceException("Error communicating with Scylla", e);
 		}
 	}
 
@@ -131,17 +145,24 @@ public class ScyllaBlobIndex : IBlobIndex
 	{
 		using TelemetrySpan scope = _tracer.BuildScyllaSpan("scylla.fetch_old_blob_index_references").SetAttribute("resource.name", $"{ns}.{id}");
 
-		if (_scyllaSessionManager.IsScylla)
+		try
 		{
-			ScyllaBlobIndexTable o;
-			o = await _mapper.SingleOrDefaultAsync<ScyllaBlobIndexTable>("SELECT namespace, blob_id, regions, references FROM blob_index WHERE namespace = ? AND blob_id = ?", ns.ToString(), new ScyllaBlobIdentifier(id));
-			return o?.References?.ToList();
+			if (_scyllaSessionManager.IsScylla)
+			{
+				ScyllaBlobIndexTable o;
+				o = await _mapper.SingleOrDefaultAsync<ScyllaBlobIndexTable>("SELECT namespace, blob_id, regions, references FROM blob_index WHERE namespace = ? AND blob_id = ?", ns.ToString(), new ScyllaBlobIdentifier(id));
+				return o?.References?.ToList();
+			}
+			else
+			{
+				CassandraBlobIndexTable o;
+				o = await _mapper.SingleOrDefaultAsync<CassandraBlobIndexTable>("SELECT namespace, blob_id, regions, references FROM blob_index WHERE namespace = ? AND blob_id = ?", ns.ToString(), id.HashData);
+				return o?.References?.ToList();
+			}
 		}
-		else
+		catch (QueryExecutionException e)
 		{
-			CassandraBlobIndexTable o;
-			o = await _mapper.SingleOrDefaultAsync<CassandraBlobIndexTable>("SELECT namespace, blob_id, regions, references FROM blob_index WHERE namespace = ? AND blob_id = ?", ns.ToString(), id.HashData);
-			return o?.References?.ToList();
+			throw new RemoteServiceException("Error communicating with Scylla", e);
 		}
 	}
 
@@ -150,14 +171,28 @@ public class ScyllaBlobIndex : IBlobIndex
 		region ??= _jupiterSettings.CurrentValue.CurrentSite;
 		using TelemetrySpan scope = _tracer.BuildScyllaSpan("scylla.remove_blob_index_region").SetAttribute("resource.name", $"{ns}.{id}");
 
-		await _mapper.DeleteAsync<ScyllaBlobIndexEntry>(new ScyllaBlobIndexEntry(ns.ToString(), id, region));
+		try
+		{
+			await _mapper.DeleteAsync<ScyllaBlobIndexEntry>(new ScyllaBlobIndexEntry(ns.ToString(), id, region));
+		}
+		catch (QueryExecutionException e)
+		{
+			throw new RemoteServiceException("Error communicating with Scylla", e);
+		}
 	}
 
 	public async Task RemoveBlobFromAllRegionsAsync(NamespaceId ns, BlobId id, CancellationToken cancellationToken = default)
 	{
 		using TelemetrySpan scope = _tracer.BuildScyllaSpan("scylla.remove_blob_index_all_regions").SetAttribute("resource.name", $"{ns}.{id}");
 
-		await _mapper.DeleteAsync<ScyllaBlobIndexEntry>("WHERE namespace = ? AND blob_id = ?", ns.ToString(), id.HashData);
+		try
+		{
+			await _mapper.DeleteAsync<ScyllaBlobIndexEntry>("WHERE namespace = ? AND blob_id = ?", ns.ToString(), id.HashData);
+		}
+		catch (QueryExecutionException e)
+		{
+			throw new RemoteServiceException("Error communicating with Scylla", e);
+		}
 	}
 	
 	public async Task<bool> BlobExistsInRegionAsync(NamespaceId ns, BlobId blobIdentifier, string? region = null, CancellationToken cancellationToken = default)
@@ -167,7 +202,14 @@ public class ScyllaBlobIndex : IBlobIndex
 		{
 			using TelemetrySpan scope = _tracer.BuildScyllaSpan("scylla.get_blob_index").SetAttribute("resource.name", $"{ns}.{blobIdentifier}");
 
-			entry = await _mapper.SingleOrDefaultAsync<ScyllaBlobIndexEntry>("WHERE namespace = ? AND blob_id = ? AND region = ?", ns.ToString(), blobIdentifier.HashData, region);
+			try
+			{
+				entry = await _mapper.SingleOrDefaultAsync<ScyllaBlobIndexEntry>("WHERE namespace = ? AND blob_id = ? AND region = ?", ns.ToString(), blobIdentifier.HashData, region);
+			}
+			catch (QueryExecutionException e)
+			{
+				throw new RemoteServiceException("Error communicating with Scylla", e);
+			}
 		}
 
 		bool blobMissing = entry == null;
@@ -197,7 +239,17 @@ public class ScyllaBlobIndex : IBlobIndex
 
 		string cqlOptions = _scyllaSessionManager.IsScylla ? "BYPASS CACHE" : "";
 
-		foreach (ScyllaBlobIndexEntry blobIndex in await _mapper.FetchAsync<ScyllaBlobIndexEntry>(cqlOptions))
+		IEnumerable<ScyllaBlobIndexEntry> entries;
+		try
+		{
+			entries = await _mapper.FetchAsync<ScyllaBlobIndexEntry>(cqlOptions);
+		}
+		catch (QueryExecutionException e)
+		{
+			throw new RemoteServiceException("Error communicating with Scylla", e);
+		}
+		
+		foreach (ScyllaBlobIndexEntry blobIndex in entries)
 		{
 			yield return (new NamespaceId(blobIndex.Namespace), new BlobId(blobIndex.BlobId));
 		}
@@ -217,7 +269,17 @@ public class ScyllaBlobIndex : IBlobIndex
 		foreach (int index in tableRangeIndices)
 		{
 			(long, long) range = tableRanges[index];
-			RowSet rowSet = await _session.ExecuteAsync(_getBlobReferenceStatement.Bind(range.Item1, range.Item2));
+
+			RowSet rowSet;
+			try
+			{
+				rowSet = await _session.ExecuteAsync(_getBlobReferenceStatement.Bind(range.Item1, range.Item2));
+			}
+			catch (QueryExecutionException e)
+			{
+				throw new RemoteServiceException("Error communicating with Scylla", e);
+			}
+			
 			foreach (Row row in rowSet)
 			{
 				string ns = row.GetValue<string>("namespace");
@@ -242,7 +304,18 @@ public class ScyllaBlobIndex : IBlobIndex
 	public async Task<List<string>> GetBlobRegionsAsync(NamespaceId ns, BlobId blob, CancellationToken cancellationToken)
 	{
 		List<string> regions = new List<string>();
-		foreach (ScyllaBlobIndexEntry blobIndex in await _mapper.FetchAsync<ScyllaBlobIndexEntry>("WHERE namespace = ? AND blob_id = ?", ns.ToString(), blob.HashData))
+
+		IEnumerable<ScyllaBlobIndexEntry> entries;
+		try
+		{
+			entries = await _mapper.FetchAsync<ScyllaBlobIndexEntry>("WHERE namespace = ? AND blob_id = ?", ns.ToString(), blob.HashData);
+		}
+		catch (QueryExecutionException e)
+		{
+			throw new RemoteServiceException("Error communicating with Scylla", e);
+		}
+		
+		foreach (ScyllaBlobIndexEntry blobIndex in entries)
 		{
 			regions.Add(blobIndex.Region);
 		}
@@ -279,7 +352,14 @@ public class ScyllaBlobIndex : IBlobIndex
 		string nsAsString = ns.ToString();
 		ScyllaBlobIncomingReference incomingReference = new ScyllaBlobIncomingReference(nsAsString, sourceBlob, targetBlob);
 
-		await _mapper.InsertAsync<ScyllaBlobIncomingReference>(incomingReference);
+		try
+		{
+			await _mapper.InsertAsync<ScyllaBlobIncomingReference>(incomingReference);
+		}
+		catch (QueryExecutionException e)
+		{
+			throw new RemoteServiceException("Error communicating with Scylla", e);
+		}
 	}
 
 	public async IAsyncEnumerable<BaseBlobReference> GetBlobReferencesAsync(NamespaceId ns, BlobId id, [EnumeratorCancellation] CancellationToken cancellationToken)
@@ -290,7 +370,18 @@ public class ScyllaBlobIndex : IBlobIndex
 		string cqlOptions = _scyllaSessionManager.IsScylla ? "BYPASS CACHE" : "";
 
 		bool noReferencesFound = true;
-		foreach (ScyllaBlobIncomingReference incomingReference in await _mapper.FetchAsync<ScyllaBlobIncomingReference>("WHERE namespace = ? AND blob_id = ? " + cqlOptions, ns.ToString(), id.HashData))
+
+		IEnumerable<ScyllaBlobIncomingReference> references;
+		try
+		{
+			references = await _mapper.FetchAsync<ScyllaBlobIncomingReference>("WHERE namespace = ? AND blob_id = ? " + cqlOptions, ns.ToString(), id.HashData);
+		}
+		catch (QueryExecutionException e)
+		{
+			throw new RemoteServiceException("Error communicating with Scylla", e);
+		}
+		
+		foreach (ScyllaBlobIncomingReference incomingReference in references)
 		{
 			noReferencesFound = false;
 			ScyllaBlobIncomingReference.BlobReferenceType referenceType = (ScyllaBlobIncomingReference.BlobReferenceType)incomingReference.ReferenceType;
@@ -353,7 +444,14 @@ public class ScyllaBlobIndex : IBlobIndex
 
 		if (referencesToRemove == null)
 		{
-			await _mapper.DeleteAsync<ScyllaBlobIncomingReference>("WHERE namespace = ? AND blob_id = ?", nsAsString, id.HashData);
+			try
+			{
+				await _mapper.DeleteAsync<ScyllaBlobIncomingReference>("WHERE namespace = ? AND blob_id = ?", nsAsString, id.HashData);
+			}
+			catch (QueryExecutionException e)
+			{
+				throw new RemoteServiceException("Error communicating with Scylla", e);
+			}
 		}
 		else
 		{
diff --git a/Engine/Source/Programs/UnrealCloudDDC/Jupiter/Implementation/Builds/ScyllaBlockStore.cs b/Engine/Source/Programs/UnrealCloudDDC/Jupiter/Implementation/Builds/ScyllaBlockStore.cs
--- a/Engine/Source/Programs/UnrealCloudDDC/Jupiter/Implementation/Builds/ScyllaBlockStore.cs
+++ b/Engine/Source/Programs/UnrealCloudDDC/Jupiter/Implementation/Builds/ScyllaBlockStore.cs
@@ -6,6 +6,7 @@ using System.Threading.Tasks;
 using Cassandra;
 using Cassandra.Mapping;
 using EpicGames.Horde.Storage;
+using Jupiter.Common;
 using Microsoft.Extensions.Options;
 using OpenTelemetry.Trace;
 
@@ -63,49 +64,80 @@ public class ScyllaBlockStore : IBlockStore
 	public async Task AddBlockToContextAsync(NamespaceId ns, BlockContext blockContext, BlobId metadataBlockId)
 	{
 		using TelemetrySpan scope = _tracer.BuildScyllaSpan("scylla.update_block_index");
-		await _mapper.InsertAsync<ScyllaBlockContextEntry>(new ScyllaBlockContextEntry(ns, blockContext, DateTime.Now, metadataBlockId));
+
+		try
+		{
+			await _mapper.InsertAsync<ScyllaBlockContextEntry>(new ScyllaBlockContextEntry(ns, blockContext, DateTime.Now, metadataBlockId));
+		}
+		catch (QueryExecutionException e)
+		{
+			throw new RemoteServiceException("Error communicating with Scylla", e);
+		}
 	}
 
 	public async Task PutBlockMetadataAsync(NamespaceId ns, BlobId blockIdentifier, BlobId metadataObjectId)
 	{
 		using TelemetrySpan scope = _tracer.BuildScyllaSpan("scylla.put_block_metadata");
 
-		await _mapper.InsertAsync<ScyllaBlockIndexEntry>(new ScyllaBlockIndexEntry(ns, blockIdentifier, metadataObjectId));
+		try
+		{
+			await _mapper.InsertAsync<ScyllaBlockIndexEntry>(new ScyllaBlockIndexEntry(ns, blockIdentifier, metadataObjectId));
+		}
+		catch (QueryExecutionException e)
+		{
+			throw new RemoteServiceException("Error communicating with Scylla", e);
+		}
 	}
 
 	public async Task<BlobId?> GetBlockMetadataAsync(NamespaceId ns, BlobId blockIdentifier)
 	{
 		using TelemetrySpan scope = _tracer.BuildScyllaSpan("scylla.get_block_metadata");
-		ScyllaBlockIndexEntry? scyllaBlockIndexEntry = await _mapper.SingleOrDefaultAsync<ScyllaBlockIndexEntry>("WHERE namespace = ? AND block_id = ?", ns.ToString(), blockIdentifier.HashData);
 
-		if (scyllaBlockIndexEntry != null)
+		try
+		{
+			ScyllaBlockIndexEntry? scyllaBlockIndexEntry = await _mapper.SingleOrDefaultAsync<ScyllaBlockIndexEntry>("WHERE namespace = ? AND block_id = ?", ns.ToString(), blockIdentifier.HashData);
+
+			if (scyllaBlockIndexEntry != null)
+			{
+				scope.SetAttribute("found", bool.TrueString);
+				return new BlobId(scyllaBlockIndexEntry.MetadataBlobId);
+			}
+			scope.SetAttribute("found", bool.FalseString);
+			return null;
+		}
+		catch (QueryExecutionException e)
 		{
-			scope.SetAttribute("found", bool.TrueString);
-			return new BlobId(scyllaBlockIndexEntry.MetadataBlobId);
+			throw new RemoteServiceException("Error communicating with Scylla", e);
 		}
-		scope.SetAttribute("found", bool.FalseString);
-		return null;
 	}
 
 	public async Task DeleteBlockAsync(NamespaceId ns, BlobId blockIdentifier)
 	{
 		using TelemetrySpan scope = _tracer.BuildScyllaSpan("scylla.delete_block_metadata");
-		List<string> foundContexts = new ();
-		foreach (ScyllaBlockContextEntry blockEntry in await _mapper.FetchAsync<ScyllaBlockContextEntry>("SELECT * FROM block_context WHERE block_id = ? ", blockIdentifier.HashData))
+		
+		try
 		{
-			if (string.Equals(blockEntry.Namespace, ns.ToString(), StringComparison.OrdinalIgnoreCase))
+			List<string> foundContexts = new ();
+			foreach (ScyllaBlockContextEntry blockEntry in await _mapper.FetchAsync<ScyllaBlockContextEntry>("SELECT * FROM block_context WHERE block_id = ? ", blockIdentifier.HashData))
 			{
-				foundContexts.Add(blockEntry.BlockContext);
+				if (string.Equals(blockEntry.Namespace, ns.ToString(), StringComparison.OrdinalIgnoreCase))
+				{
+					foundContexts.Add(blockEntry.BlockContext);
+				}
 			}
-		}
 
-		foreach (string blockContext in foundContexts)
+			foreach (string blockContext in foundContexts)
+			{
+				await _mapper.DeleteAsync<ScyllaBlockContextEntry>("WHERE namespace = ? AND block_context = ? AND block_id = ?", ns.ToString(), blockContext, blockIdentifier.HashData);
+			}
+
+			scope.SetAttribute("UsedInCountOfContexts", foundContexts.Count);
+			await _mapper.DeleteAsync<ScyllaBlockIndexEntry>("WHERE namespace = ? AND block_id = ?", ns.ToString(), blockIdentifier.HashData);
+		}
+		catch (QueryExecutionException e)
 		{
-			await _mapper.DeleteAsync<ScyllaBlockContextEntry>("WHERE namespace = ? AND block_context = ? AND block_id = ?", ns.ToString(), blockContext, blockIdentifier.HashData);
+			throw new RemoteServiceException("Error communicating with Scylla", e);
 		}
-
-		scope.SetAttribute("UsedInCountOfContexts", foundContexts.Count);
-		await _mapper.DeleteAsync<ScyllaBlockIndexEntry>("WHERE namespace = ? AND block_id = ?", ns.ToString(), blockIdentifier.HashData);
 	}
 }
 
diff --git a/Engine/Source/Programs/UnrealCloudDDC/Jupiter/Implementation/Builds/ScyllaBuildStore.cs b/Engine/Source/Programs/UnrealCloudDDC/Jupiter/Implementation/Builds/ScyllaBuildStore.cs
--- a/Engine/Source/Programs/UnrealCloudDDC/Jupiter/Implementation/Builds/ScyllaBuildStore.cs
+++ b/Engine/Source/Programs/UnrealCloudDDC/Jupiter/Implementation/Builds/ScyllaBuildStore.cs
@@ -9,6 +9,7 @@ using Cassandra;
 using Cassandra.Mapping;
 using EpicGames.Horde.Storage;
 using EpicGames.Serialization;
+using Jupiter.Common;
 using Microsoft.Extensions.Options;
 using OpenTelemetry.Trace;
 
@@ -58,7 +59,18 @@ public class ScyllaBuildStore : IBuildStore
 	public async IAsyncEnumerable<BuildMetadata> ListBuildsAsync(NamespaceId ns, BucketId bucket, bool includeTTL = false)
 	{
 		using TelemetrySpan scope = _tracer.BuildScyllaSpan("scylla.list_builds");
-		foreach (ScyllaBuildEntry buildEntry in await _mapper.FetchAsync<ScyllaBuildEntry>("WHERE namespace = ? AND bucket_id = ?", ns.ToString(), bucket.ToString()))
+		
+		IEnumerable<ScyllaBuildEntry> entries;
+		try
+		{
+			entries = await _mapper.FetchAsync<ScyllaBuildEntry>("WHERE namespace = ? AND bucket_id = ?", ns.ToString(), bucket.ToString());
+		}
+		catch (QueryExecutionException e)
+		{
+			throw new RemoteServiceException("Error communicating with Scylla", e);
+		}
+		
+		foreach (ScyllaBuildEntry buildEntry in entries)
 		{
 			BuildMetadata? metadata = buildEntry.ToBuildMetadata();
 			// ignore invalid metadata
@@ -81,21 +93,42 @@ public class ScyllaBuildStore : IBuildStore
 	{
 		using TelemetrySpan scope = _tracer.BuildScyllaSpan("scylla.put_build");
 
-		await _mapper.InsertAsync<ScyllaBuildEntry>(new ScyllaBuildEntry(ns, bucket, buildId, buildObject), insertNulls: false, ttl: (int)ttl);
+		try
+		{
+			await _mapper.InsertAsync<ScyllaBuildEntry>(new ScyllaBuildEntry(ns, bucket, buildId, buildObject), insertNulls: false, ttl: (int)ttl);
+		}
+		catch (QueryExecutionException e)
+		{
+			throw new RemoteServiceException("Error communicating with Scylla", e);
+		}
 	}
 
 	public async Task FinalizeBuildAsync(NamespaceId ns, BucketId bucket, CbObjectId buildId, uint ttl)
 	{
-		await _mapper.UpdateAsync<ScyllaBuildEntry>("USING TTL ? SET is_finalized = true WHERE namespace = ? and bucket_id = ? AND build_id = ?", (int)ttl, ns.ToString(), bucket.ToString(), buildId.ToByteArray());
+		try
+		{
+			await _mapper.UpdateAsync<ScyllaBuildEntry>("USING TTL ? SET is_finalized = true WHERE namespace = ? and bucket_id = ? AND build_id = ?", (int)ttl, ns.ToString(), bucket.ToString(), buildId.ToByteArray());
+		}
+		catch (QueryExecutionException e)
+		{
+			throw new RemoteServiceException("Error communicating with Scylla", e);
+		}
 	}
 
 	public async Task<BuildRecord?> GetBuildAsync(NamespaceId ns, BucketId bucket, CbObjectId buildId)
 	{
-		ScyllaBuildEntry? buildEntry = await _mapper.SingleOrDefaultAsync<ScyllaBuildEntry>("WHERE namespace = ? AND bucket_id = ? AND build_id = ?", ns.ToString(), bucket.ToString(), buildId.ToByteArray());
+		try
+		{
+			ScyllaBuildEntry? buildEntry = await _mapper.SingleOrDefaultAsync<ScyllaBuildEntry>("WHERE namespace = ? AND bucket_id = ? AND build_id = ?", ns.ToString(), bucket.ToString(), buildId.ToByteArray());
 
-		if (buildEntry != null && buildEntry.BuildObject != null) 
+			if (buildEntry != null && buildEntry.BuildObject != null) 
+			{
+				return new BuildRecord {BuildObject = new CbObject(buildEntry.BuildObject), IsFinalized = buildEntry.IsFinalized};
+			}
+		}
+		catch (QueryExecutionException e)
 		{
-			return new BuildRecord {BuildObject = new CbObject(buildEntry.BuildObject), IsFinalized = buildEntry.IsFinalized};
+			throw new RemoteServiceException("Error communicating with Scylla", e);
 		}
 
 		return null;
@@ -105,13 +138,32 @@ public class ScyllaBuildStore : IBuildStore
 	{
 		using TelemetrySpan scope = _tracer.BuildScyllaSpan("scylla.put_build_part");
 
-		await _mapper.InsertAsync<ScyllaBuildPartEntry>(new ScyllaBuildPartEntry(ns, bucket, buildId, partId, partName), insertNulls: false, ttl: (int)ttl);
+		try
+		{
+			await _mapper.InsertAsync<ScyllaBuildPartEntry>(new ScyllaBuildPartEntry(ns, bucket, buildId, partId, partName), insertNulls: false, ttl: (int)ttl);
+		}
+		catch (QueryExecutionException e)
+		{
+			throw new RemoteServiceException("Error communicating with Scylla", e);
+		}
 	}
 
 	public async IAsyncEnumerable<(string, CbObjectId)> GetBuildPartsAsync(NamespaceId ns, BucketId bucket, CbObjectId buildId)
 	{
 		using TelemetrySpan scope = _tracer.BuildScyllaSpan("scylla.get_build_parts");
-		foreach (ScyllaBuildPartEntry buildPartEntry in await _mapper.FetchAsync<ScyllaBuildPartEntry>("WHERE namespace = ? AND bucket_id = ? AND build_id = ?", ns.ToString(), bucket.ToString(), buildId.ToByteArray()))
+		
+		IEnumerable<ScyllaBuildPartEntry> entries;
+		try
+		{
+			entries = await _mapper.FetchAsync<ScyllaBuildPartEntry>("WHERE namespace = ? AND bucket_id = ? AND build_id = ?", ns.ToString(),
+				bucket.ToString(), buildId.ToByteArray());
+		}
+		catch (QueryExecutionException e)
+		{
+			throw new RemoteServiceException("Error communicating with Scylla", e);
+		}
+		
+		foreach (ScyllaBuildPartEntry buildPartEntry in entries)
 		{
 			yield return (buildPartEntry.PartName, new CbObjectId(buildPartEntry.PartId));
 		}
@@ -121,17 +173,24 @@ public class ScyllaBuildStore : IBuildStore
 	{
 		using TelemetrySpan scope = _tracer.BuildScyllaSpan("scylla.get_build_ttl");
 
-		BoundStatement? boundStatement = _getTtlStatement.Bind(ns.ToString(), bucket.ToString(), buildId.ToByteArray());
-		RowSet rowSet = await _session.ExecuteAsync(boundStatement);
-		foreach (Row row in rowSet)
+		try
 		{
-			// we only expect one row here so we just return once we have a value
-
-			if (row["ttl"] is int ttl)
+			BoundStatement? boundStatement = _getTtlStatement.Bind(ns.ToString(), bucket.ToString(), buildId.ToByteArray());
+			RowSet rowSet = await _session.ExecuteAsync(boundStatement);
+			foreach (Row row in rowSet)
 			{
-				return (uint)ttl;
+				// we only expect one row here so we just return once we have a value
+
+				if (row["ttl"] is int ttl)
+				{
+					return (uint)ttl;
+				}
 			}
 		}
+		catch (QueryExecutionException e)
+		{
+			throw new RemoteServiceException("Error communicating with Scylla", e);
+		}
 
 		return null;
 	}
@@ -142,7 +201,17 @@ public class ScyllaBuildStore : IBuildStore
 
 		string cqlOptions = _scyllaSessionManager.IsScylla ? "BYPASS CACHE" : "";
 
-		foreach (ScyllaBuildEntry buildEntry in await _mapper.FetchAsync<ScyllaBuildEntry>(cqlOptions))
+		IEnumerable<ScyllaBuildEntry> entries;
+		try
+		{
+			entries = await _mapper.FetchAsync<ScyllaBuildEntry>(cqlOptions);
+		}
+		catch (QueryExecutionException e)
+		{
+			throw new RemoteServiceException("Error communicating with Scylla", e);
+		}
+		
+		foreach (ScyllaBuildEntry buildEntry in entries)
 		{
 			yield return (new NamespaceId(buildEntry.Namespace), new BucketId(buildEntry.BucketId), new CbObjectId(buildEntry.BuildId));
 		}
@@ -150,18 +219,36 @@ public class ScyllaBuildStore : IBuildStore
 
 	public async Task UpdateTTL(NamespaceId ns, BucketId bucket, CbObjectId buildId, uint ttl)
 	{
-		ScyllaBuildEntry? buildEntry = await _mapper.SingleOrDefaultAsync<ScyllaBuildEntry>("WHERE namespace = ? AND bucket_id = ? AND build_id = ?", ns.ToString(), bucket.ToString(), buildId.ToByteArray());
+		try
+		{
+			ScyllaBuildEntry? buildEntry = await _mapper.SingleOrDefaultAsync<ScyllaBuildEntry>("WHERE namespace = ? AND bucket_id = ? AND build_id = ?", ns.ToString(), bucket.ToString(), buildId.ToByteArray());
 
-		if (buildEntry == null)
+			if (buildEntry == null)
+			{
+				// failed to find the full build, very unlikely, should have failed earlier
+				return;
+			}
+
+			// scylla does not support updating the ttl, instead we reinsert the same value with a new ttl
+			await _mapper.InsertAsync<ScyllaBuildEntry>(buildEntry, false, (int)ttl);
+		}
+		catch (QueryExecutionException e)
 		{
-			// failed to find the full build, very unlikely, should have failed earlier
-			return;
+			throw new RemoteServiceException("Error communicating with Scylla", e);
 		}
 
-		// scylla does not support updating the ttl, instead we reinsert the same value with a new ttl
-		await _mapper.InsertAsync<ScyllaBuildEntry>(buildEntry, false, (int)ttl);
+		IEnumerable<ScyllaBuildPartEntry> entries;
+		try
+		{
+			entries = await _mapper.FetchAsync<ScyllaBuildPartEntry>("WHERE namespace = ? AND bucket_id = ? AND build_id = ?", ns.ToString(),
+				bucket.ToString(), buildId.ToByteArray());
+		}
+		catch (QueryExecutionException e)
+		{
+			throw new RemoteServiceException("Error communicating with Scylla", e);
+		}
 
-		foreach (ScyllaBuildPartEntry buildPartEntry in await _mapper.FetchAsync<ScyllaBuildPartEntry>("WHERE namespace = ? AND bucket_id = ? AND build_id = ?", ns.ToString(), bucket.ToString(), buildId.ToByteArray()))
+		foreach (ScyllaBuildPartEntry buildPartEntry in entries)
 		{
 			await _mapper.InsertAsync<ScyllaBuildPartEntry>(buildPartEntry, false, (int)ttl);
 		}
@@ -169,8 +256,15 @@ public class ScyllaBuildStore : IBuildStore
 
 	public async Task DeleteBuild(NamespaceId ns, BucketId bucket, CbObjectId buildId)
 	{
-		await _mapper.DeleteAsync<ScyllaBuildPartEntry>("WHERE namespace = ? AND bucket_id = ? AND build_id = ?", ns.ToString(), bucket.ToString(), buildId.ToByteArray());
-		await _mapper.DeleteAsync<ScyllaBuildEntry>("WHERE namespace = ? AND bucket_id = ? AND build_id = ?", ns.ToString(), bucket.ToString(), buildId.ToByteArray());
+		try
+		{
+			await _mapper.DeleteAsync<ScyllaBuildPartEntry>("WHERE namespace = ? AND bucket_id = ? AND build_id = ?", ns.ToString(), bucket.ToString(), buildId.ToByteArray());
+			await _mapper.DeleteAsync<ScyllaBuildEntry>("WHERE namespace = ? AND bucket_id = ? AND build_id = ?", ns.ToString(), bucket.ToString(), buildId.ToByteArray());
+		}
+		catch (QueryExecutionException e)
+		{
+			throw new RemoteServiceException("Error communicating with Scylla", e);
+		}
 	}
 }
 
diff --git a/Engine/Source/Programs/UnrealCloudDDC/Jupiter/Implementation/Objects/ScyllaReferencesStore.cs b/Engine/Source/Programs/UnrealCloudDDC/Jupiter/Implementation/Objects/ScyllaReferencesStore.cs
--- a/Engine/Source/Programs/UnrealCloudDDC/Jupiter/Implementation/Objects/ScyllaReferencesStore.cs
+++ b/Engine/Source/Programs/UnrealCloudDDC/Jupiter/Implementation/Objects/ScyllaReferencesStore.cs
@@ -105,38 +105,50 @@ namespace Jupiter.Implementation
 			using TelemetrySpan scope = _tracer.BuildScyllaSpan("scylla.get").SetAttribute("resource.name", $"{ns}.{bucket}.{name}");
 			scope.SetAttribute("BypassCache", false);
 
-			ScyllaObject? o;
-			bool includePayload = (fieldFlags & IReferencesStore.FieldFlags.IncludePayload) != 0;
-			if (includePayload)
+			try
 			{
-				o = await _mapper.SingleOrDefaultAsync<ScyllaObject>("WHERE namespace = ? AND bucket = ? AND name = ?", ns.ToString(), bucket.ToString(), name.ToString());
-			}
-			else
-			{
-				string cqlOptions = "";
-				if (_scyllaSessionManager.IsScylla && opFlags.HasFlag(IReferencesStore.OperationFlags.BypassCache))
+				ScyllaObject? o;
+				bool includePayload = (fieldFlags & IReferencesStore.FieldFlags.IncludePayload) != 0;
+				if (includePayload)
 				{
-					// BYPASS CACHE is a scylla specific extension to disable populating the cache, should be ignored by other cassandra dbs
-					cqlOptions = "BYPASS CACHE";
-					scope.SetAttribute("BypassCache", true);
+					o = await _mapper.SingleOrDefaultAsync<ScyllaObject>("WHERE namespace = ? AND bucket = ? AND name = ?", ns.ToString(),
+						bucket.ToString(), name.ToString());
 				}
+				else
+				{
+					string cqlOptions = "";
+					if (_scyllaSessionManager.IsScylla && opFlags.HasFlag(IReferencesStore.OperationFlags.BypassCache))
+					{
+						// BYPASS CACHE is a scylla specific extension to disable populating the cache, should be ignored by other cassandra dbs
+						cqlOptions = "BYPASS CACHE";
+						scope.SetAttribute("BypassCache", true);
+					}
 
-				// fetch everything except for the inline blob which is quite large
-				o = await _mapper.SingleOrDefaultAsync<ScyllaObject>($"SELECT namespace, bucket, name , payload_hash, is_finalized, last_access_time FROM objects WHERE namespace = ? AND bucket = ? AND name = ? {cqlOptions}", ns.ToString(), bucket.ToString(), name.ToString());
-			}
+					// fetch everything except for the inline blob which is quite large
+					o = await _mapper.SingleOrDefaultAsync<ScyllaObject>(
+						$"SELECT namespace, bucket, name , payload_hash, is_finalized, last_access_time FROM objects WHERE namespace = ? AND bucket = ? AND name = ? {cqlOptions}",
+						ns.ToString(), bucket.ToString(), name.ToString());
+				}
 
-			if (o == null)
-			{
-				throw new RefNotFoundException(ns, bucket, name);
+				if (o == null)
+				{
+					throw new RefNotFoundException(ns, bucket, name);
+				}
+
+				if (o.IsRequiredFieldIsMissing())
+				{
+					_logger.LogWarning("Partial object found {Namespace} {Bucket} {Name} cleaning up stale data. Will be ignored.", ns, bucket, name);
+					await DeleteAsync(ns, bucket, name, cancellationToken);
+					throw new RefNotFoundException(ns, bucket, name);
+				}
+
+				return new RefRecord(new NamespaceId(o.Namespace!), new BucketId(o.Bucket!), new RefId(o.Name!), o.LastAccessTime, o.InlinePayload,
+					o.PayloadHash!.AsBlobIdentifier(), o.IsFinalized!.Value);
 			}
-			if (o.IsRequiredFieldIsMissing())
+			catch (QueryExecutionException e)
 			{
-				_logger.LogWarning("Partial object found {Namespace} {Bucket} {Name} cleaning up stale data. Will be ignored.", ns, bucket, name);
-				await DeleteAsync(ns, bucket, name, cancellationToken);
-				throw new RefNotFoundException(ns, bucket, name);
+				throw new RemoteServiceException("Error communicating with Scylla", e);
 			}
-
-			return new RefRecord(new NamespaceId(o.Namespace!), new BucketId(o.Bucket!), new RefId(o.Name!), o.LastAccessTime, o.InlinePayload, o.PayloadHash!.AsBlobIdentifier(), o.IsFinalized!.Value);
 		}
 
 		public async Task PutAsync(NamespaceId ns, BucketId bucket, RefId name, BlobId blobHash, byte[] blob, bool isFinalized, bool allowOverwrite = false, CancellationToken cancellationToken = default)
@@ -154,46 +166,58 @@ namespace Jupiter.Implementation
 				ttl = (int)policy.DefaultTTL.TotalSeconds;
 			}
 
-			Task? insertLastAccess = gcMethod == NamespacePolicy.StoragePoolGCMethod.LastAccess ? _mapper.InsertAsync<ScyllaObjectLastAccess>(new ScyllaObjectLastAccess(ns, bucket, name, DateTime.Now)) : null;
-
-			bool allowOverwrites = policy.AllowOverwritesOfRefs || allowOverwrite;
-			if (!allowOverwrites)
+			try
 			{
-				bool insertRef = false;
-				try
+				Task? insertLastAccess = gcMethod == NamespacePolicy.StoragePoolGCMethod.LastAccess
+					? _mapper.InsertAsync<ScyllaObjectLastAccess>(new ScyllaObjectLastAccess(ns, bucket, name, DateTime.Now))
+					: null;
+
+				bool allowOverwrites = policy.AllowOverwritesOfRefs || allowOverwrite;
+				if (!allowOverwrites)
 				{
-					RefRecord oldRecord = await GetAsync(ns, bucket, name, IReferencesStore.FieldFlags.IncludePayload, IReferencesStore.OperationFlags.None, cancellationToken);
+					bool insertRef = false;
+					try
+					{
+						RefRecord oldRecord = await GetAsync(ns, bucket, name, IReferencesStore.FieldFlags.IncludePayload,
+							IReferencesStore.OperationFlags.None, cancellationToken);
+
+						if (!oldRecord.BlobIdentifier.Equals(blobHash))
+						{
+							// blob was not the same, e.g. we attempted to change the value, this is not allowed
+							throw new RefAlreadyExistsException(ns, bucket, name, oldRecord);
+						}
 
-					if (!oldRecord.BlobIdentifier.Equals(blobHash))
+						// new record is the same as the old, this is fine there is no need to write anything
+					}
+					catch (RefNotFoundException)
 					{
-						// blob was not the same, e.g. we attempted to change the value, this is not allowed
-						throw new RefAlreadyExistsException(ns, bucket, name, oldRecord);
+						insertRef = true;
+					}
+
+					if (insertRef)
+					{
+						await _mapper.InsertAsync<ScyllaObject>(new ScyllaObject(ns, bucket, name, blob, blobHash, isFinalized), ttl: ttl,
+							insertNulls: false);
 					}
-					
-					// new record is the same as the old, this is fine there is no need to write anything
 				}
-				catch (RefNotFoundException)
+				else
 				{
-					insertRef = true;
+					// if we allow overwrites we do the insert without the "if not exists" check which is much faster
+					await _mapper.InsertAsync<ScyllaObject>(new ScyllaObject(ns, bucket, name, blob, blobHash, isFinalized), ttl: ttl,
+						insertNulls: false);
 				}
 
-				if (insertRef)
+				if (insertLastAccess != null)
 				{
-					await _mapper.InsertAsync<ScyllaObject>(new ScyllaObject(ns, bucket, name, blob, blobHash, isFinalized), ttl: ttl, insertNulls: false);
+					await insertLastAccess;
 				}
-			}
-			else
-			{
-				// if we allow overwrites we do the insert without the "if not exists" check which is much faster
-				await _mapper.InsertAsync<ScyllaObject>(new ScyllaObject(ns, bucket, name, blob, blobHash, isFinalized), ttl: ttl, insertNulls: false);
-			}
 
-			if (insertLastAccess != null)
+				await addBucketTask;
+			}
+			catch (QueryExecutionException e)
 			{
-				await insertLastAccess;
+				throw new RemoteServiceException("Error communicating with Scylla", e);
 			}
-
-			await addBucketTask;
 		}
 
 		public async Task FinalizeAsync(NamespaceId ns, BucketId bucket, RefId name, BlobId blobIdentifier, CancellationToken cancellationToken)
@@ -208,36 +232,58 @@ namespace Jupiter.Implementation
 				ttl = (int)policy.DefaultTTL.TotalSeconds;
 			}
 
-			if (ttl != null)
+			try
 			{
-				await _mapper.UpdateAsync<ScyllaObject>("USING TTL ? SET is_finalized=true WHERE namespace=? AND bucket=? AND name=?", ttl.Value, ns.ToString(), bucket.ToString(), name.ToString());
+				if (ttl != null)
+				{
+					await _mapper.UpdateAsync<ScyllaObject>("USING TTL ? SET is_finalized=true WHERE namespace=? AND bucket=? AND name=?", ttl.Value, ns.ToString(), bucket.ToString(), name.ToString());
+				}
+				else
+				{
+					await _mapper.UpdateAsync<ScyllaObject>("SET is_finalized=true WHERE namespace=? AND bucket=? AND name=?", ns.ToString(), bucket.ToString(), name.ToString());
+				}
 			}
-			else
+			catch (QueryExecutionException e)
 			{
-				await _mapper.UpdateAsync<ScyllaObject>("SET is_finalized=true WHERE namespace=? AND bucket=? AND name=?", ns.ToString(), bucket.ToString(), name.ToString());
+				throw new RemoteServiceException("Error communicating with Scylla", e);
 			}
 		}
 
 		public async Task<DateTime?> GetLastAccessTimeAsync(NamespaceId ns, BucketId bucket, RefId key, CancellationToken cancellationToken)
 		{
 			using TelemetrySpan scope = _tracer.BuildScyllaSpan("scylla.get_last_access_time").SetAttribute("resource.name", $"{ns}.{bucket}.{key}");
-			ScyllaObjectLastAccess? lastAccessRecord = await _mapper.SingleOrDefaultAsync<ScyllaObjectLastAccess>("WHERE namespace = ? AND bucket = ? AND name = ?", ns.ToString(), bucket.ToString(), key.ToString());
-			return lastAccessRecord?.LastAccessTime;
+			
+			try
+			{
+				ScyllaObjectLastAccess? lastAccessRecord = await _mapper.SingleOrDefaultAsync<ScyllaObjectLastAccess>("WHERE namespace = ? AND bucket = ? AND name = ?", ns.ToString(), bucket.ToString(), key.ToString());
+				return lastAccessRecord?.LastAccessTime;
+			}
+			catch (QueryExecutionException e)
+			{
+				throw new RemoteServiceException("Error communicating with Scylla", e);
+			}
 		}
 
 		public async Task UpdateLastAccessTimeAsync(NamespaceId ns, BucketId bucket, RefId name, DateTime lastAccessTime, CancellationToken cancellationToken)
 		{
 			using TelemetrySpan scope = _tracer.BuildScyllaSpan("scylla.update_last_access_time");
 
-			Task? updateObjectLastAccessTask = null;
-			updateObjectLastAccessTask = _mapper.InsertAsync<ScyllaObjectLastAccess>(new ScyllaObjectLastAccess(ns, bucket, name, lastAccessTime));
+			try
+			{
+				Task? updateObjectLastAccessTask = null;
+				updateObjectLastAccessTask = _mapper.InsertAsync<ScyllaObjectLastAccess>(new ScyllaObjectLastAccess(ns, bucket, name, lastAccessTime));
+
+				if (_settings.CurrentValue.UpdateLegacyLastAccessField)
+				{
+					await _mapper.UpdateAsync<ScyllaObject>("SET last_access_time = ? WHERE namespace = ? AND bucket = ? AND name = ?", lastAccessTime, ns.ToString(), bucket.ToString(), name.ToString());
+				}
 
-			if (_settings.CurrentValue.UpdateLegacyLastAccessField)
+				await updateObjectLastAccessTask;
+			}
+			catch (QueryExecutionException e)
 			{
-				await _mapper.UpdateAsync<ScyllaObject>("SET last_access_time = ? WHERE namespace = ? AND bucket = ? AND name = ?", lastAccessTime, ns.ToString(), bucket.ToString(), name.ToString());
+				throw new RemoteServiceException("Error communicating with Scylla", e);
 			}
-
-			await updateObjectLastAccessTask;
 		}
 
 		public async IAsyncEnumerable<(NamespaceId, BucketId, RefId, DateTime)> GetRecordsAsync([EnumeratorCancellation] CancellationToken cancellationToken)
@@ -312,7 +358,7 @@ namespace Jupiter.Implementation
 						_logger.LogWarning("Cassandra read timeouts, attempted {Attempts} attempts now we give up.",
 							retryAttempts);
 						// we have failed to many times, rethrow the exception and abort to avoid stalling here for ever
-						throw timeoutException;
+						throw new RemoteServiceException("Error communicating with Scylla", timeoutException);
 					}
 				} while (!rowSet.IsFullyFetched);
 			}
@@ -386,7 +432,7 @@ namespace Jupiter.Implementation
 						_logger.LogWarning("Cassandra read timeouts, attempted {Attempts} attempts now we give up.",
 							retryAttempts);
 						// we have failed to many times, rethrow the exception and abort to avoid stalling here for ever
-						throw timeoutException;
+						throw new RemoteServiceException("Error communicating with Scylla", timeoutException);
 					}
 				} while (!rowSet.IsFullyFetched);
 			}
@@ -542,64 +588,94 @@ namespace Jupiter.Implementation
 		{
 			using TelemetrySpan scope = _tracer.BuildScyllaSpan("scylla.delete_record").SetAttribute("resource.name", $"{ns}.{bucket}.{key}");
 
-			Task? lastAccessDeleteTask = _mapper.DeleteAsync<ScyllaObjectLastAccess>("WHERE namespace=? AND bucket=? AND name=?", ns.ToString(), bucket.ToString(), key.ToString());
+			try
+			{
+				Task? lastAccessDeleteTask = _mapper.DeleteAsync<ScyllaObjectLastAccess>("WHERE namespace=? AND bucket=? AND name=?", ns.ToString(), bucket.ToString(), key.ToString());
 
-			await _mapper.DeleteAsync<ScyllaObject>("WHERE namespace=? AND bucket=? AND name=?", ns.ToString(), bucket.ToString(), key.ToString());
+				await _mapper.DeleteAsync<ScyllaObject>("WHERE namespace=? AND bucket=? AND name=?", ns.ToString(), bucket.ToString(), key.ToString());
 
-			await lastAccessDeleteTask;
+				await lastAccessDeleteTask;
 
-			return true;
+				return true;
+			}
+			catch (QueryExecutionException e)
+			{
+				throw new RemoteServiceException("Error communicating with Scylla", e);
+			}
 		}
 
 		public async Task<long> DropNamespaceAsync(NamespaceId ns, CancellationToken cancellationToken)
 		{
 			using TelemetrySpan scope = _tracer.BuildScyllaSpan("scylla.delete_namespace");
-			RowSet rowSet = await _session.ExecuteAsync(new SimpleStatement("SELECT bucket, name FROM objects WHERE namespace = ? ALLOW FILTERING;", ns.ToString()));
-			long deletedCount = 0;
-			foreach (Row row in rowSet)
+
+			try
 			{
-				string bucket = row.GetValue<string>("bucket");
-				string name = row.GetValue<string>("name");
+				RowSet rowSet = await _session.ExecuteAsync(new SimpleStatement("SELECT bucket, name FROM objects WHERE namespace = ? ALLOW FILTERING;", ns.ToString()));
+				long deletedCount = 0;
+				foreach (Row row in rowSet)
+				{
+					string bucket = row.GetValue<string>("bucket");
+					string name = row.GetValue<string>("name");
 
-				await DeleteAsync(ns, new BucketId(bucket), new RefId(name), cancellationToken);
+					await DeleteAsync(ns, new BucketId(bucket), new RefId(name), cancellationToken);
 
-				deletedCount++;
-			}
+					deletedCount++;
+				}
 
-			// remove the tracking in the buckets table as well
-			await _session.ExecuteAsync(new SimpleStatement("DELETE FROM buckets WHERE namespace = ?", ns.ToString()));
-			await _session.ExecuteAsync(new SimpleStatement("DELETE FROM buckets_v2 WHERE namespace = ?", ns.ToString()));
+				// remove the tracking in the buckets table as well
+				await _session.ExecuteAsync(new SimpleStatement("DELETE FROM buckets WHERE namespace = ?", ns.ToString()));
+				await _session.ExecuteAsync(new SimpleStatement("DELETE FROM buckets_v2 WHERE namespace = ?", ns.ToString()));
 
-			return deletedCount;
+				return deletedCount;
+			}
+			catch (QueryExecutionException e)
+			{
+				throw new RemoteServiceException("Error communicating with Scylla", e);
+			}
 		}
 
 		public async Task<long> DeleteBucketAsync(NamespaceId ns, BucketId bucket, CancellationToken cancellationToken)
 		{
 			using TelemetrySpan scope = _tracer.BuildScyllaSpan("scylla.delete_bucket");
 
-			RowSet rowSet = await _session.ExecuteAsync(new SimpleStatement("SELECT name FROM objects WHERE namespace = ? AND bucket = ? ALLOW FILTERING;", ns.ToString(), bucket.ToString()));
-			long deletedCount = 0;
-			foreach (Row row in rowSet)
+			try
 			{
-				string name = row.GetValue<string>("name");
+				RowSet rowSet = await _session.ExecuteAsync(new SimpleStatement("SELECT name FROM objects WHERE namespace = ? AND bucket = ? ALLOW FILTERING;", ns.ToString(), bucket.ToString()));
+				long deletedCount = 0;
+				foreach (Row row in rowSet)
+				{
+					string name = row.GetValue<string>("name");
 
-				await DeleteAsync(ns, bucket, new RefId(name), cancellationToken);
-				deletedCount++;
-			}
+					await DeleteAsync(ns, bucket, new RefId(name), cancellationToken);
+					deletedCount++;
+				}
 
-			// remove the tracking in the buckets table as well
-			await _mapper.DeleteAsync<ScyllaBucket>(new ScyllaBucket(ns, bucket));
+				// remove the tracking in the buckets table as well
+				await _mapper.DeleteAsync<ScyllaBucket>(new ScyllaBucket(ns, bucket));
 
-			return deletedCount;
+				return deletedCount;
+			}
+			catch (QueryExecutionException e)
+			{
+				throw new RemoteServiceException("Error communicating with Scylla", e);
+			}
 		}
 
 		public async Task UpdateTTL(NamespaceId ns, BucketId bucket, RefId key, uint ttl, CancellationToken cancellationToken = default)
 		{
 			using TelemetrySpan scope = _tracer.BuildScyllaSpan("scylla.update_ttl");
-			ScyllaObject? o = await _mapper.SingleOrDefaultAsync<ScyllaObject>("WHERE namespace=? AND bucket=? AND name=?", ns.ToString(), bucket.ToString(), key.ToString());
-			if (o != null)
+
+			try
 			{
-				await _mapper.InsertAsync<ScyllaObject>(o, false, (int)ttl);
+				ScyllaObject? o = await _mapper.SingleOrDefaultAsync<ScyllaObject>("WHERE namespace=? AND bucket=? AND name=?", ns.ToString(), bucket.ToString(), key.ToString());
+				if (o != null)
+				{
+					await _mapper.InsertAsync<ScyllaObject>(o, false, (int)ttl);
+				}
+			}
+			catch (QueryExecutionException e)
+			{
+				throw new RemoteServiceException("Error communicating with Scylla", e);
 			}
 		}
 
@@ -612,76 +688,93 @@ namespace Jupiter.Implementation
 			bool alreadyAdded = addedBuckets.Contains(bucket);
 			if (!alreadyAdded)
 			{
-				Task addTask = _mapper.InsertAsync<ScyllaBucket>(new ScyllaBucket(ns, bucket));
-				addedBuckets.Add(bucket);
-				await addTask;
+				try
+				{
+					Task addTask = _mapper.InsertAsync<ScyllaBucket>(new ScyllaBucket(ns, bucket));
+					addedBuckets.Add(bucket);
+					await addTask;
+				}
+				catch (QueryExecutionException e)
+				{
+					throw new RemoteServiceException("Error communicating with Scylla", e);
+				}
 			}
 		}
 
 		public async Task CleanupTrackedStateAsync()
 		{
-			_logger.LogInformation("Starting internal scylla cleanup. ");
-			DateTime cleanupStart = DateTime.Now;
-			ConcurrentDictionary<NamespaceId, ConcurrentDictionary<BucketId, bool>> namespacesWithRefs = new();
-
-			// build lookup of which namespaces and buckets are actually references by live objects
+			try
 			{
-				RowSet rowSet = await _session.ExecuteAsync(_getObjectsBucketsStatement.Bind());
+				_logger.LogInformation("Starting internal scylla cleanup. ");
+				DateTime cleanupStart = DateTime.Now;
+				ConcurrentDictionary<NamespaceId, ConcurrentDictionary<BucketId, bool>> namespacesWithRefs = new();
 
-				foreach (Row row in rowSet)
+				// build lookup of which namespaces and buckets are actually references by live objects
 				{
-					if (rowSet.GetAvailableWithoutFetching() == 0)
-					{
-						await rowSet.FetchMoreResultsAsync();
-					}
+					RowSet rowSet = await _session.ExecuteAsync(_getObjectsBucketsStatement.Bind());
 
-					NamespaceId ns = new NamespaceId(row.GetValue<string>(0));
-					BucketId bucket = new BucketId(row.GetValue<string>(1));
+					foreach (Row row in rowSet)
+					{
+						if (rowSet.GetAvailableWithoutFetching() == 0)
+						{
+							await rowSet.FetchMoreResultsAsync();
+						}
 
-					ConcurrentDictionary<BucketId, bool> bucketTracking = namespacesWithRefs.GetOrAdd(ns, new ConcurrentDictionary<BucketId, bool>());
-					bucketTracking.TryAdd(bucket, true);
-				}
-			}
+						NamespaceId ns = new NamespaceId(row.GetValue<string>(0));
+						BucketId bucket = new BucketId(row.GetValue<string>(1));
 
-			// enumerate our cached namespaces and buckets and determine if any of them are old
-			await foreach (NamespaceId ns in GetNamespacesAsync(CancellationToken.None))
-			{
-				bool refFound = false;
-				if (!namespacesWithRefs.TryGetValue(ns, out ConcurrentDictionary<BucketId, bool>? foundBuckets))
-				{
-					// namespace does not exist, still need to processes all buckets to remove their tracking
-					foundBuckets = new ConcurrentDictionary<BucketId, bool>();
+						ConcurrentDictionary<BucketId, bool> bucketTracking =
+							namespacesWithRefs.GetOrAdd(ns, new ConcurrentDictionary<BucketId, bool>());
+						bucketTracking.TryAdd(bucket, true);
+					}
 				}
 
-				await foreach (BucketId bucket in GetBucketsAsync(ns, CancellationToken.None))
+				// enumerate our cached namespaces and buckets and determine if any of them are old
+				await foreach (NamespaceId ns in GetNamespacesAsync(CancellationToken.None))
 				{
-					// if the bucket had a ref we should not remove its tracking
-					if (foundBuckets.ContainsKey(bucket))
+					bool refFound = false;
+					if (!namespacesWithRefs.TryGetValue(ns, out ConcurrentDictionary<BucketId, bool>? foundBuckets))
 					{
-						refFound = true;
-						continue;
+						// namespace does not exist, still need to processes all buckets to remove their tracking
+						foundBuckets = new ConcurrentDictionary<BucketId, bool>();
 					}
-					_logger.LogWarning("Namespace {Namespace} Bucket {Bucket} has no refs. Should be removed from tracking.", ns, bucket);
 
-					if (_settings.CurrentValue.CleanupOldTrackedState)
+					await foreach (BucketId bucket in GetBucketsAsync(ns, CancellationToken.None))
 					{
-						await _session.ExecuteAsync(new SimpleStatement("DELETE FROM buckets_v2 WHERE namespace = ? and bucket = ?", ns.ToString(), bucket.ToString()));
-					}
-				}
+						// if the bucket had a ref we should not remove its tracking
+						if (foundBuckets.ContainsKey(bucket))
+						{
+							refFound = true;
+							continue;
+						}
 
-				if (!refFound)
-				{
-					_logger.LogWarning("Namespace {Namespace} has no refs. Should be removed from tracking.", ns);
+						_logger.LogWarning("Namespace {Namespace} Bucket {Bucket} has no refs. Should be removed from tracking.", ns, bucket);
 
-					if (_settings.CurrentValue.CleanupOldTrackedState)
+						if (_settings.CurrentValue.CleanupOldTrackedState)
+						{
+							await _session.ExecuteAsync(new SimpleStatement("DELETE FROM buckets_v2 WHERE namespace = ? and bucket = ?",
+								ns.ToString(), bucket.ToString()));
+						}
+					}
+
+					if (!refFound)
 					{
-						await _session.ExecuteAsync(new SimpleStatement("DELETE FROM buckets_v2 WHERE namespace = ?", ns.ToString()));
+						_logger.LogWarning("Namespace {Namespace} has no refs. Should be removed from tracking.", ns);
+
+						if (_settings.CurrentValue.CleanupOldTrackedState)
+						{
+							await _session.ExecuteAsync(new SimpleStatement("DELETE FROM buckets_v2 WHERE namespace = ?", ns.ToString()));
+						}
 					}
 				}
-			}
 
-			TimeSpan cleanupDuration = DateTime.Now - cleanupStart;
-			_logger.LogInformation("Finished internal scylla cleanup. Cleanup took: {CleanupDuration}", cleanupDuration);
+				TimeSpan cleanupDuration = DateTime.Now - cleanupStart;
+				_logger.LogInformation("Finished internal scylla cleanup. Cleanup took: {CleanupDuration}", cleanupDuration);
+			}
+			catch (QueryExecutionException e)
+			{
+				throw new RemoteServiceException("Error communicating with Scylla", e);
+			}
 		}
 
 		public async IAsyncEnumerable<RefId> GetRecordsInBucketAsync(NamespaceId ns, BucketId bucket, [EnumeratorCancellation] CancellationToken cancellationToken)
diff --git a/Engine/Source/Programs/UnrealCloudDDC/Jupiter/Implementation/TransactionLog/ScyllaReplicationLog.cs b/Engine/Source/Programs/UnrealCloudDDC/Jupiter/Implementation/TransactionLog/ScyllaReplicationLog.cs
--- a/Engine/Source/Programs/UnrealCloudDDC/Jupiter/Implementation/TransactionLog/ScyllaReplicationLog.cs
+++ b/Engine/Source/Programs/UnrealCloudDDC/Jupiter/Implementation/TransactionLog/ScyllaReplicationLog.cs
@@ -80,7 +80,16 @@ namespace Jupiter.Implementation
 
 		public async IAsyncEnumerable<NamespaceId> GetNamespacesAsync()
 		{
-			IEnumerable<ScyllaNamespace> namespaces = await _mapper.FetchAsync<ScyllaNamespace>("");
+			IEnumerable<ScyllaNamespace> namespaces;
+			
+			try
+			{
+				namespaces = await _mapper.FetchAsync<ScyllaNamespace>("");
+			}
+			catch (QueryExecutionException e)
+			{
+				throw new RemoteServiceException("Error communicating with Scylla", e);
+			}
 
 			foreach (ScyllaNamespace scyllaNamespace in namespaces)
 			{
@@ -94,11 +103,19 @@ namespace Jupiter.Implementation
 
 			Task addNamespaceTask = PotentiallyAddNamespaceAsync(ns);
 			DateTime timeBucket = timestamp.GetValueOrDefault(DateTime.UtcNow);
-			ScyllaReplicationLogEvent log = new ScyllaReplicationLogEvent(ns.ToString(), bucket.ToString(), key.ToString(), timeBucket, ScyllaReplicationLogEvent.OpType.Added, objectBlob);
-			await _mapper.InsertAsync<ScyllaReplicationLogEvent>(log, insertNulls: false, ttl: (int)_settings.CurrentValue.ReplicationLogTimeToLive.TotalSeconds);
 
-			await addNamespaceTask;
-			return (log.GetReplicationBucketIdentifier(), log.ReplicationId);
+			try
+			{
+				ScyllaReplicationLogEvent log = new ScyllaReplicationLogEvent(ns.ToString(), bucket.ToString(), key.ToString(), timeBucket, ScyllaReplicationLogEvent.OpType.Added, objectBlob);
+				await _mapper.InsertAsync<ScyllaReplicationLogEvent>(log, insertNulls: false, ttl: (int)_settings.CurrentValue.ReplicationLogTimeToLive.TotalSeconds);
+
+				await addNamespaceTask;
+				return (log.GetReplicationBucketIdentifier(), log.ReplicationId);
+			}
+			catch (QueryExecutionException e)
+			{
+				throw new RemoteServiceException("Error communicating with Scylla", e);
+			}
 		}
 
 		public async Task<(string, Guid)> InsertDeleteEventAsync(NamespaceId ns, BucketId bucket, RefId key, DateTime? timestamp)
@@ -107,16 +124,31 @@ namespace Jupiter.Implementation
 
 			Task addNamespaceTask = PotentiallyAddNamespaceAsync(ns);
 			DateTime timeBucket = timestamp.GetValueOrDefault(DateTime.UtcNow);
-			ScyllaReplicationLogEvent log = new ScyllaReplicationLogEvent(ns.ToString(), bucket.ToString(), key.ToString(), timeBucket, ScyllaReplicationLogEvent.OpType.Deleted, null);
-			await _mapper.InsertAsync<ScyllaReplicationLogEvent>(log, insertNulls: false, ttl: (int)_settings.CurrentValue.ReplicationLogTimeToLive.TotalSeconds);
 
-			await addNamespaceTask;
-			return (log.GetReplicationBucketIdentifier(), log.ReplicationId);
+			try
+			{
+				ScyllaReplicationLogEvent log = new ScyllaReplicationLogEvent(ns.ToString(), bucket.ToString(), key.ToString(), timeBucket, ScyllaReplicationLogEvent.OpType.Deleted, null);
+				await _mapper.InsertAsync<ScyllaReplicationLogEvent>(log, insertNulls: false, ttl: (int)_settings.CurrentValue.ReplicationLogTimeToLive.TotalSeconds);
+
+				await addNamespaceTask;
+				return (log.GetReplicationBucketIdentifier(), log.ReplicationId);
+			}
+			catch (QueryExecutionException e)
+			{
+				throw new RemoteServiceException("Error communicating with Scylla", e);
+			}
 		}
 
 		private async Task PotentiallyAddNamespaceAsync(NamespaceId ns)
 		{
-			await _mapper.InsertAsync(new ScyllaNamespace(ns.ToString()));
+			try
+			{
+				await _mapper.InsertAsync(new ScyllaNamespace(ns.ToString()));
+			}
+			catch (QueryExecutionException e)
+			{
+				throw new RemoteServiceException("Error communicating with Scylla", e);
+			}
 		}
 
 		public async IAsyncEnumerable<ReplicationLogEvent> GetAsync(NamespaceId ns, string? lastBucket, Guid? lastEvent)
@@ -148,7 +180,15 @@ namespace Jupiter.Implementation
 				// we have found the bucket to resume from, we should no do any more filtering
 				lastBucket = null;
 
-				IEnumerable<ScyllaReplicationLogEvent> events = await _mapper.FetchAsync<ScyllaReplicationLogEvent>("WHERE namespace = ? AND replication_bucket = ?", ns.ToString(), bucketField);
+				IEnumerable<ScyllaReplicationLogEvent> events;
+				try
+				{
+					events = await _mapper.FetchAsync<ScyllaReplicationLogEvent>("WHERE namespace = ? AND replication_bucket = ?", ns.ToString(), bucketField);
+				}
+				catch (QueryExecutionException e)
+				{
+					throw new RemoteServiceException("Error communicating with Scylla", e);
+				}
 
 				bool skipEvents = lastEvent.HasValue;
 				foreach (ScyllaReplicationLogEvent scyllaReplicationLog in events)
@@ -182,13 +222,20 @@ namespace Jupiter.Implementation
 
 			if (!bucketFound)
 			{
-				ScyllaNamespace? namespaces = await _mapper.SingleOrDefaultAsync<ScyllaNamespace>("WHERE namespace = ?", ns.ToString());
-				if (namespaces == null)
+				try
 				{
-					throw new NamespaceNotFoundException(ns);
-				}
+					ScyllaNamespace? namespaces = await _mapper.SingleOrDefaultAsync<ScyllaNamespace>("WHERE namespace = ?", ns.ToString());
+					if (namespaces == null)
+					{
+						throw new NamespaceNotFoundException(ns);
+					}
 
-				throw new IncrementalLogNotAvailableException();
+					throw new IncrementalLogNotAvailableException();
+				}
+				catch (QueryExecutionException e)
+				{
+					throw new RemoteServiceException("Error communicating with Scylla", e);
+				}
 			}
 		}
 
@@ -261,28 +308,42 @@ namespace Jupiter.Implementation
 
 		public async Task AddSnapshotAsync(SnapshotInfo snapshotHeader)
 		{
-			await _mapper.InsertAsync<ScyllaSnapshot>(new ScyllaSnapshot(snapshotHeader.SnapshottedNamespace.ToString(), snapshotHeader.BlobNamespace.ToString(), TimeUuid.NewId(), snapshotHeader.SnapshotBlob));
-			await CleanupSnapshotsAsync(snapshotHeader.SnapshottedNamespace);
+			try
+			{
+				await _mapper.InsertAsync<ScyllaSnapshot>(new ScyllaSnapshot(snapshotHeader.SnapshottedNamespace.ToString(), snapshotHeader.BlobNamespace.ToString(), TimeUuid.NewId(), snapshotHeader.SnapshotBlob));
+				await CleanupSnapshotsAsync(snapshotHeader.SnapshottedNamespace);
+			}
+			catch (QueryExecutionException e)
+			{
+				throw new RemoteServiceException("Error communicating with Scylla", e);
+			}
 		}
 
 		public async Task CleanupSnapshotsAsync(NamespaceId ns)
 		{
-			// determine if we have to many snapshots and remove the oldest ones if we do
-			RowSet rowSet = await _session.ExecuteAsync(new SimpleStatement("SELECT Count(*) FROM replication_snapshot WHERE namespace = ?", ns.ToString()));
-			Row row = rowSet.First();
-			int countOfSnapshots = (int)(long)row[0];
-			if (countOfSnapshots > _settings.CurrentValue.MaxSnapshotsPerNamespace)
+			try
 			{
-				int deleteCount = countOfSnapshots - _settings.CurrentValue.MaxSnapshotsPerNamespace;
-
-				List<ScyllaSnapshot> snapshots = (await _mapper.FetchAsync<ScyllaSnapshot>("WHERE namespace = ?", ns.ToString())).ToList();
-				for (int i = 0; i < deleteCount; i++)
+				// determine if we have to many snapshots and remove the oldest ones if we do
+				RowSet rowSet = await _session.ExecuteAsync(new SimpleStatement("SELECT Count(*) FROM replication_snapshot WHERE namespace = ?", ns.ToString()));
+				Row row = rowSet.First();
+				int countOfSnapshots = (int)(long)row[0];
+				if (countOfSnapshots > _settings.CurrentValue.MaxSnapshotsPerNamespace)
 				{
-					// since snapshots are sorted newest first we delete from the end
-					ScyllaSnapshot snapshotToDelete = snapshots[^(i + 1)];
-					await _mapper.DeleteAsync<ScyllaSnapshot>(snapshotToDelete);
+					int deleteCount = countOfSnapshots - _settings.CurrentValue.MaxSnapshotsPerNamespace;
+
+					List<ScyllaSnapshot> snapshots = (await _mapper.FetchAsync<ScyllaSnapshot>("WHERE namespace = ?", ns.ToString())).ToList();
+					for (int i = 0; i < deleteCount; i++)
+					{
+						// since snapshots are sorted newest first we delete from the end
+						ScyllaSnapshot snapshotToDelete = snapshots[^(i + 1)];
+						await _mapper.DeleteAsync<ScyllaSnapshot>(snapshotToDelete);
+					}
 				}
 			}
+			catch (QueryExecutionException e)
+			{
+				throw new RemoteServiceException("Error communicating with Scylla", e);
+			}
 		}
 
 		public async Task<SnapshotInfo?> GetLatestSnapshotAsync(NamespaceId ns)
@@ -294,7 +355,15 @@ namespace Jupiter.Implementation
 
 		public async IAsyncEnumerable<SnapshotInfo> GetSnapshotsAsync(NamespaceId ns)
 		{
-			IEnumerable<ScyllaSnapshot> snapshots = await _mapper.FetchAsync<ScyllaSnapshot>("WHERE namespace = ?", ns.ToString());
+			IEnumerable<ScyllaSnapshot> snapshots;
+			try
+			{
+				snapshots = await _mapper.FetchAsync<ScyllaSnapshot>("WHERE namespace = ?", ns.ToString());
+			}
+			catch (QueryExecutionException e)
+			{
+				throw new RemoteServiceException("Error communicating with Scylla", e);
+			}
 
 			foreach (ScyllaSnapshot snapshot in snapshots)
 			{
@@ -304,16 +373,30 @@ namespace Jupiter.Implementation
 
 		public async Task UpdateReplicatorStateAsync(NamespaceId ns, string replicatorName, ReplicatorState newState)
 		{
-			await _mapper.UpdateAsync<ScyllaReplicationState>(new ScyllaReplicationState(ns, replicatorName, newState.LastBucket, newState.LastEvent));
+			try
+			{
+				await _mapper.UpdateAsync<ScyllaReplicationState>(new ScyllaReplicationState(ns, replicatorName, newState.LastBucket, newState.LastEvent));
+			}
+			catch (QueryExecutionException e)
+			{
+				throw new RemoteServiceException("Error communicating with Scylla", e);
+			}
 		}
 
 		public async Task<ReplicatorState?> GetReplicatorStateAsync(NamespaceId ns, string name)
 		{
 			using TelemetrySpan scope = _tracer.BuildScyllaSpan("scylla.get_replicator_state");
 
-			ScyllaReplicationState? replicationState = await _mapper.FirstOrDefaultAsync<ScyllaReplicationState>("WHERE namespace = ? AND name = ?", ns.ToString(), name);
+			try
+			{
+				ScyllaReplicationState? replicationState = await _mapper.FirstOrDefaultAsync<ScyllaReplicationState>("WHERE namespace = ? AND name = ?", ns.ToString(), name);
 
-			return replicationState?.ToReplicatorState();
+				return replicationState?.ToReplicatorState();
+			}
+			catch (QueryExecutionException e)
+			{
+				throw new RemoteServiceException("Error communicating with Scylla", e);
+			}
 		}
 
 		public async Task<(string, Guid)> InsertAddBlobEventAsync(NamespaceId ns, BlobId objectBlob, DateTime? timestamp = null, BucketId? bucketHint = null)
@@ -322,11 +405,19 @@ namespace Jupiter.Implementation
 
 			Task addNamespaceTask = PotentiallyAddNamespaceAsync(ns);
 			DateTime timeBucket = timestamp.GetValueOrDefault(DateTime.UtcNow);
-			ScyllaBlobReplicationLogEvent log = new ScyllaBlobReplicationLogEvent(ns.ToString(), objectBlob, timeBucket, ScyllaBlobReplicationLogEvent.OpType.Added, bucketHint);
-			await _mapper.InsertAsync<ScyllaBlobReplicationLogEvent>(log, insertNulls: false, ttl: (int)_settings.CurrentValue.ReplicationLogTimeToLive.TotalSeconds);
 
-			await addNamespaceTask;
-			return (log.GetReplicationBucketIdentifier(), log.ReplicationId);
+			try
+			{
+				ScyllaBlobReplicationLogEvent log = new ScyllaBlobReplicationLogEvent(ns.ToString(), objectBlob, timeBucket, ScyllaBlobReplicationLogEvent.OpType.Added, bucketHint);
+				await _mapper.InsertAsync<ScyllaBlobReplicationLogEvent>(log, insertNulls: false, ttl: (int)_settings.CurrentValue.ReplicationLogTimeToLive.TotalSeconds);
+
+				await addNamespaceTask;
+				return (log.GetReplicationBucketIdentifier(), log.ReplicationId);
+			}
+			catch (QueryExecutionException e)
+			{
+				throw new RemoteServiceException("Error communicating with Scylla", e);
+			}
 		}
 
 		public async IAsyncEnumerable<BlobReplicationLogEvent> GetBlobEventsAsync(NamespaceId ns, string replicationBucket)
@@ -342,9 +433,17 @@ namespace Jupiter.Implementation
 				throw new IncrementalLogNotAvailableException();
 			}
 
-			long bucketId = replicationBucket.FromReplicationBucketIdentifier().ToFileTimeUtc();
-			IEnumerable<ScyllaBlobReplicationLogEvent> events = await _mapper.FetchAsync<ScyllaBlobReplicationLogEvent>("WHERE namespace = ? AND replication_bucket = ?", ns.ToString(), bucketId);
-
+			IEnumerable<ScyllaBlobReplicationLogEvent> events;
+			try
+			{
+				long bucketId = replicationBucket.FromReplicationBucketIdentifier().ToFileTimeUtc();
+				events = await _mapper.FetchAsync<ScyllaBlobReplicationLogEvent>("WHERE namespace = ? AND replication_bucket = ?", ns.ToString(), bucketId);
+			}
+			catch (QueryExecutionException e)
+			{
+				throw new RemoteServiceException("Error communicating with Scylla", e);
+			}
+			
 			bool eventFound = false;
 			foreach (ScyllaBlobReplicationLogEvent scyllaReplicationLog in events)
 			{
@@ -361,10 +460,17 @@ namespace Jupiter.Implementation
 
 			if (!eventFound)
 			{
-				ScyllaNamespace? namespaces = await _mapper.SingleOrDefaultAsync<ScyllaNamespace>("WHERE namespace = ?", ns.ToString());
-				if (namespaces == null)
+				try
+				{
+					ScyllaNamespace? namespaces = await _mapper.SingleOrDefaultAsync<ScyllaNamespace>("WHERE namespace = ?", ns.ToString());
+					if (namespaces == null)
+					{
+						throw new NamespaceNotFoundException(ns);
+					}
+				}
+				catch (QueryExecutionException e)
 				{
-					throw new NamespaceNotFoundException(ns);
+					throw new RemoteServiceException("Error communicating with Scylla", e);
 				}
 			}
 		}
