Cloud DDC - Use ContainerResource autoscaling metrics in Helm templates

diff --git a/Engine/Source/Programs/UnrealCloudDDC/Helm/UnrealCloudDDC/templates/hpa.yaml b/Engine/Source/Programs/UnrealCloudDDC/Helm/UnrealCloudDDC/templates/hpa.yaml
--- a/Engine/Source/Programs/UnrealCloudDDC/Helm/UnrealCloudDDC/templates/hpa.yaml
+++ b/Engine/Source/Programs/UnrealCloudDDC/Helm/UnrealCloudDDC/templates/hpa.yaml
@@ -22,17 +22,19 @@ spec:
   maxReplicas: {{ .Values.autoscaling.hpa.maxReplicas }}
   metrics:
   {{- if .Values.autoscaling.hpa.targetCPUUtilizationPercentage }}
-    - type: Resource
-      resource:
+    - type: ContainerResource
+      containerResource:
         name: cpu
+        container: {{ include "common.names.name" . }}
         target:
           averageUtilization: {{ .Values.autoscaling.hpa.targetCPUUtilizationPercentage }}
           type: Utilization
   {{- end }}
   {{- if .Values.autoscaling.hpa.targetMemoryUtilizationPercentage }}
-    - type: Resource
-      resource:
+    - type: ContainerResource
+      containerResource:
         name: memory
+        container: {{ include "common.names.name" . }}
         target:
           averageUtilization: {{ .Values.autoscaling.hpa.targetMemoryUtilizationPercentage }}
           type: Utilization
