Cloud DDC - Increased the number of builds that are allowed to be scanned by default during cross bucket searches from 5000 to 20 000

diff --git a/Engine/Source/Programs/UnrealCloudDDC/Jupiter/Controllers/BuildsController.cs b/Engine/Source/Programs/UnrealCloudDDC/Jupiter/Controllers/BuildsController.cs
--- a/Engine/Source/Programs/UnrealCloudDDC/Jupiter/Controllers/BuildsController.cs
+++ b/Engine/Source/Programs/UnrealCloudDDC/Jupiter/Controllers/BuildsController.cs
@@ -2111,7 +2111,7 @@ namespace Jupiter.Controllers
 		public int Skip { get; set; } = 0;
 
 		[CbField("max")]
-		public int Max { get; set; } = 5000;
+		public int Max { get; set; } = 20_000;
 
 		[CbField("maxTimeMS")]
 		public uint MaxTimeMs { get; set; } = Debugger.IsAttached ? 0u : 120_000; // default to 2 minutes max per search unless debugging in which case we disable it
