Cloud DDC - Catch S3 slowdown issues during build part put and return 429

diff --git a/Engine/Source/Programs/UnrealCloudDDC/Jupiter/Controllers/BuildsController.cs b/Engine/Source/Programs/UnrealCloudDDC/Jupiter/Controllers/BuildsController.cs
--- a/Engine/Source/Programs/UnrealCloudDDC/Jupiter/Controllers/BuildsController.cs
+++ b/Engine/Source/Programs/UnrealCloudDDC/Jupiter/Controllers/BuildsController.cs
@@ -981,8 +981,15 @@ namespace Jupiter.Controllers
 
 			BlockContext blockContext = BlockContext.FromObject(buildRecord.BuildObject);
 
-			uint ttl = (uint)_namespacePolicyResolver.GetPoliciesForNs(ns).DefaultTTL.TotalSeconds;
-			await _buildStore.PutBuildPartAsync(ns, bucket, buildId, partId, partName, ttl);
+			try
+			{
+				uint ttl = (uint)_namespacePolicyResolver.GetPoliciesForNs(ns).DefaultTTL.TotalSeconds;
+				await _buildStore.PutBuildPartAsync(ns, bucket, buildId, partId, partName, ttl);
+			}
+			catch (ResourceHasToManyRequestsException)
+			{
+				return StatusCode(StatusCodes.Status429TooManyRequests);
+			}
 
 			try
 			{
