Cloud DDC - Make reference resolver v2 the default resolver. We have been using this internally for a year, handles objects with more attachments better.

diff --git a/Engine/Source/Programs/UnrealCloudDDC/Jupiter/Implementation/IReferenceResolver.cs b/Engine/Source/Programs/UnrealCloudDDC/Jupiter/Implementation/IReferenceResolver.cs
--- a/Engine/Source/Programs/UnrealCloudDDC/Jupiter/Implementation/IReferenceResolver.cs
+++ b/Engine/Source/Programs/UnrealCloudDDC/Jupiter/Implementation/IReferenceResolver.cs
@@ -520,7 +520,7 @@ namespace Jupiter.Implementation
 
 	public class ReferenceResolverSettings
 	{
-		public bool UseV2ReferenceResolving { get; set; } = false;
+		public bool UseV2ReferenceResolving { get; set; } = true;
 		public int ChunkSize { get; set; } = 64;
 	}
 
