Cloud DDC - Retrying verifications that the build exists to give scylla some time to flush these writes

diff --git a/Engine/Source/Programs/UnrealCloudDDC/Jupiter/Controllers/BuildsController.cs b/Engine/Source/Programs/UnrealCloudDDC/Jupiter/Controllers/BuildsController.cs
--- a/Engine/Source/Programs/UnrealCloudDDC/Jupiter/Controllers/BuildsController.cs
+++ b/Engine/Source/Programs/UnrealCloudDDC/Jupiter/Controllers/BuildsController.cs
@@ -317,7 +317,7 @@ namespace Jupiter.Controllers
 
 			try
 			{
-				BuildRecord? buildRecord = await _buildStore.GetBuildAsync(ns, bucket, buildId);
+				BuildRecord? buildRecord = await GetBuildWithRetriesAsync(ns, bucket, buildId);
 
 				if (buildRecord == null)
 				{
@@ -648,7 +648,7 @@ namespace Jupiter.Controllers
 
 			try
 			{
-				BuildRecord? build = await _buildStore.GetBuildAsync(ns, bucket, buildId);
+				BuildRecord? build = await GetBuildWithRetriesAsync(ns, bucket, buildId);
 				if (build == null)
 				{
 					return NotFound($"Build \"{build}\" in bucket {bucket} and namespace {ns}");
@@ -967,7 +967,7 @@ namespace Jupiter.Controllers
 				return Problem(e.Message, null, (int)HttpStatusCode.RequestTimeout);
 			}
 
-			BuildRecord? buildRecord = await _buildStore.GetBuildAsync(ns, bucket, buildId);
+			BuildRecord? buildRecord = await GetBuildWithRetriesAsync(ns, bucket, buildId);
 			if (buildRecord == null)
 			{
 				throw new NotImplementedException($"No build found for namespace {ns} bucket {bucket} name {buildId}");
@@ -1068,10 +1068,11 @@ namespace Jupiter.Controllers
 				Tracer.CurrentSpan.SetAttribute("bucket", bucket.ToString());
 				Tracer.CurrentSpan.SetAttribute("namespace", ns.ToString());
 
+				// do not fetch with retries here as the build should already have been PUT so it's very odd for it to not be findable at this point
 				BuildRecord? buildRecord = await _buildStore.GetBuildAsync(ns, bucket, buildId);
 				if (buildRecord == null)
 				{
-					throw new NotImplementedException($"No build found for namespace {ns} bucket {bucket} name {buildId}");
+					return NotFound($"No build found for namespace {ns} bucket {bucket} name {buildId}");
 				}
 
 				// build is finalized, we should not modify it
@@ -1135,6 +1136,30 @@ namespace Jupiter.Controllers
 			}
 		}
 
+		private async Task<BuildRecord?> GetBuildWithRetriesAsync(NamespaceId ns, BucketId bucket, CbObjectId buildId)
+		{
+			const int MaxRetries = 3;
+			TimeSpan[] delayTimes = new [] {TimeSpan.FromSeconds(0.5), TimeSpan.FromSeconds(3), TimeSpan.FromSeconds(0)};
+
+			// We read this build right after submitting it in some cases which due to Scylla being eventually consistent can mean that the build doesn't exist
+			// As it should be unusual to fetch a build that doesn't exist we just retry a few times to make sure it doesn't settle into existing
+			for (int i = 0; i < MaxRetries; i++)
+			{
+				BuildRecord? buildRecord = await _buildStore.GetBuildAsync(ns, bucket, buildId);
+				if (buildRecord == null)
+				{
+					// wait a while and try again
+					await Task.Delay(delayTimes[i]);
+					continue;
+				}
+
+				return buildRecord;
+
+			}
+
+			return null;
+		}
+
 		private async Task DoFinalizeBuildPartAsync(NamespaceId ns, BlockContext blockContext, List<BlobId> blockMetadataIds, CancellationToken cancellationToken)
 		{
 			await Parallel.ForEachAsync(blockMetadataIds, cancellationToken, async (metadataBlobId, token) =>
@@ -1617,7 +1642,7 @@ namespace Jupiter.Controllers
 				return result;
 			}
 
-			BuildRecord? buildRecord = await _buildStore.GetBuildAsync(ns, bucket, buildId);
+			BuildRecord? buildRecord = await GetBuildWithRetriesAsync(ns, bucket, buildId);
 			if (buildRecord == null)
 			{
 				return NotFound($"Build {buildId} in bucket {bucket} and namespace {ns}");
@@ -1699,7 +1724,7 @@ namespace Jupiter.Controllers
 				return result;
 			}
 
-			BuildRecord? buildRecord = await _buildStore.GetBuildAsync(ns, bucket, buildId);
+			BuildRecord? buildRecord = await GetBuildWithRetriesAsync(ns, bucket, buildId);
 			if (buildRecord == null)
 			{
 				return NotFound($"Build {buildId} in bucket {bucket} and namespace {ns}");
