Cloud DDC - Create new CancellationTokenSource if TryReset() fails

diff --git a/Engine/Source/Programs/UnrealCloudDDC/Jupiter/Implementation/Replication/BlobsReplicator.cs b/Engine/Source/Programs/UnrealCloudDDC/Jupiter/Implementation/Replication/BlobsReplicator.cs
--- a/Engine/Source/Programs/UnrealCloudDDC/Jupiter/Implementation/Replication/BlobsReplicator.cs
+++ b/Engine/Source/Programs/UnrealCloudDDC/Jupiter/Implementation/Replication/BlobsReplicator.cs
@@ -27,7 +27,6 @@ namespace Jupiter.Implementation
 		private readonly string _name;
 		private readonly ILogger _logger;
 		private readonly ManualResetEvent _replicationFinishedEvent = new ManualResetEvent(true);
-		private readonly CancellationTokenSource _replicationTokenSource = new CancellationTokenSource();
 		private readonly ReplicatorSettings _replicatorSettings;
 		private readonly IBlobService _blobService;
 		private readonly IReplicationLog _replicationLog;
@@ -38,6 +37,7 @@ namespace Jupiter.Implementation
 		private readonly NamespaceId _namespace;
 		private BlobsState _replicationState;
 		private bool _replicationRunning;
+		private CancellationTokenSource _replicationTokenSource = new CancellationTokenSource();
 		private bool _disposed = false;
 
 		private static Histogram<long>? s_replicatedCounter;
@@ -180,7 +180,12 @@ namespace Jupiter.Implementation
 				
 				_logger.LogDebug("Replicator: {Name} is starting a run.", _name);
 
-				_replicationTokenSource.TryReset();
+				if (!_replicationTokenSource.TryReset())
+				{
+					_replicationTokenSource.Dispose();
+					_replicationTokenSource = new CancellationTokenSource();
+				}
+				
 				_replicationRunning = true;
 				_replicationFinishedEvent.Reset();
 				CancellationToken replicationToken = _replicationTokenSource.Token;
diff --git a/Engine/Source/Programs/UnrealCloudDDC/Jupiter/Implementation/Replication/RefsReplicator.cs b/Engine/Source/Programs/UnrealCloudDDC/Jupiter/Implementation/Replication/RefsReplicator.cs
--- a/Engine/Source/Programs/UnrealCloudDDC/Jupiter/Implementation/Replication/RefsReplicator.cs
+++ b/Engine/Source/Programs/UnrealCloudDDC/Jupiter/Implementation/Replication/RefsReplicator.cs
@@ -29,7 +29,6 @@ namespace Jupiter.Implementation
 		private readonly string _name;
 		private readonly ILogger _logger;
 		private readonly ManualResetEvent _replicationFinishedEvent = new ManualResetEvent(true);
-		private readonly CancellationTokenSource _replicationTokenSource = new CancellationTokenSource();
 		private readonly ReplicatorSettings _replicatorSettings;
 		private readonly IBlobService _blobService;
 		private readonly IReplicationLog _replicationLog;
@@ -41,6 +40,7 @@ namespace Jupiter.Implementation
 		private readonly NamespaceId _namespace;
 		private RefsState _refsState;
 		private bool _replicationRunning;
+		private CancellationTokenSource _replicationTokenSource = new CancellationTokenSource();
 		private bool _disposed = false;
 
 		private static Histogram<long>? s_replicatedCounter;
@@ -157,7 +157,12 @@ namespace Jupiter.Implementation
 
 				_logger.LogDebug("Replicator: {Name} is starting a run.", _name);
 
-				_replicationTokenSource.TryReset();
+				if (!_replicationTokenSource.TryReset())
+				{
+					_replicationTokenSource.Dispose();
+					_replicationTokenSource = new CancellationTokenSource();
+				}
+				
 				_replicationRunning = true;
 				_replicationFinishedEvent.Reset();
 				CancellationToken replicationToken = _replicationTokenSource.Token;
