Cloud DDC - Added retry of blob replication on http io exception (should already have been covered by io exception but this is at the very least clearer now) and added more logging when error happens to try and nail down why our replication fails sometimes.

diff --git a/Engine/Source/Programs/UnrealCloudDDC/Jupiter/Implementation/Replication/BlobsReplicator.cs b/Engine/Source/Programs/UnrealCloudDDC/Jupiter/Implementation/Replication/BlobsReplicator.cs
--- a/Engine/Source/Programs/UnrealCloudDDC/Jupiter/Implementation/Replication/BlobsReplicator.cs
+++ b/Engine/Source/Programs/UnrealCloudDDC/Jupiter/Implementation/Replication/BlobsReplicator.cs
@@ -536,18 +536,21 @@ namespace Jupiter.Implementation
 				}
 				catch (AggregateException e)
 				{
-					if (e.InnerException is not IOException)
+					if (e.InnerException is not IOException && e.InnerException is not HttpIOException)
 					{
+						_logger.LogError(e, "Failed to replicate {Blob} in {Namespace} due to unknown exception", blob, ns);
 						throw;
 					}
 
+					_logger.LogWarning("Failed to replicate {Blob} in {Namespace} due to retriable exception {Exception}", blob, ns, e);
 					lastException = e;
 				}
 				catch (HttpRequestException e)
 				{
 					// rethrow unknown exceptions
-					if (e.InnerException is not IOException)
+					if (e.InnerException is not IOException && e.InnerException is not HttpIOException)
 					{
+						_logger.LogError(e, "Failed to replicate {Blob} in {Namespace} due to unknown exception", blob, ns);
 						throw;
 					}
 
@@ -559,6 +562,8 @@ namespace Jupiter.Implementation
 						// delay the retry with a second if blob is missing, as its really not expected for it to actually miss, likely we just haven't reached consistency on the fact that it does exist
 						await Task.Delay(1000, cancellationToken);
 					}
+
+					_logger.LogWarning("Failed to replicate {Blob} in {Namespace} due to retriable exception {Exception}.", blob, ns, e);
 				}
 			}
 
