Cloud DDC - Return 503 when unable to retrieve a ContentId from Scylla

diff --git a/Engine/Source/Programs/UnrealCloudDDC/Jupiter/Controllers/ReferencesController.cs b/Engine/Source/Programs/UnrealCloudDDC/Jupiter/Controllers/ReferencesController.cs
--- a/Engine/Source/Programs/UnrealCloudDDC/Jupiter/Controllers/ReferencesController.cs
+++ b/Engine/Source/Programs/UnrealCloudDDC/Jupiter/Controllers/ReferencesController.cs
@@ -985,6 +985,15 @@ namespace Jupiter.Controllers
 				await Response.Body.WriteAsync(payload);
 				return new EmptyResult();
 			}
+			catch (AggregateException e)
+			{
+				if (e.InnerExceptions.Any(exception => exception.GetType() == typeof(RemoteServiceException)))
+				{
+					return Problem(e.Message, null, (int)HttpStatusCode.ServiceUnavailable);
+				}
+
+				throw;
+			}
 			catch (RemoteServiceException e)
 			{
 				return Problem(e.Message, null, (int)HttpStatusCode.ServiceUnavailable);
diff --git a/Engine/Source/Programs/UnrealCloudDDC/Jupiter/Implementation/ContentId/ScyllaContentIdStore.cs b/Engine/Source/Programs/UnrealCloudDDC/Jupiter/Implementation/ContentId/ScyllaContentIdStore.cs
--- a/Engine/Source/Programs/UnrealCloudDDC/Jupiter/Implementation/ContentId/ScyllaContentIdStore.cs
+++ b/Engine/Source/Programs/UnrealCloudDDC/Jupiter/Implementation/ContentId/ScyllaContentIdStore.cs
@@ -9,6 +9,7 @@ using System.Threading.Tasks;
 using Cassandra;
 using Cassandra.Mapping;
 using EpicGames.Horde.Storage;
+using Jupiter.Common;
 using Microsoft.Extensions.Options;
 using OpenTelemetry.Trace;
 
@@ -118,6 +119,11 @@ namespace Jupiter.Implementation
 						executionProfile = DefaultProfile;
 						continue;
 					}
+					catch (QueryExecutionException e)
+					{
+						// Other communication issues occurred
+						throw new RemoteServiceException("Error communicating with Scylla", e);
+					}
 
 					break;
 				}
