Cloud DDC - Removed unused cancellation token in blobs replication and made sure we use the more  appropriate cancellation tokens in a few places

diff --git a/Engine/Source/Programs/UnrealCloudDDC/Jupiter/Implementation/Replication/BlobsReplicator.cs b/Engine/Source/Programs/UnrealCloudDDC/Jupiter/Implementation/Replication/BlobsReplicator.cs
--- a/Engine/Source/Programs/UnrealCloudDDC/Jupiter/Implementation/Replication/BlobsReplicator.cs
+++ b/Engine/Source/Programs/UnrealCloudDDC/Jupiter/Implementation/Replication/BlobsReplicator.cs
@@ -252,9 +252,6 @@ namespace Jupiter.Implementation
 
 			_logger.LogInformation("{Name} Starting blob replication maxParallelism: {MaxParallelism} Last Bucket {LastBucket}", _name, maxParallelism, lastBucket);
 
-			using CancellationTokenSource cancellationTokenSource = new CancellationTokenSource();
-			using CancellationTokenSource linkedTokenSource = CancellationTokenSource.CreateLinkedTokenSource(cancellationTokenSource.Token, replicationToken);
-
 			if (replicationToken.IsCancellationRequested)
 			{
 				return countOfReplicationsDone;
@@ -275,9 +272,12 @@ namespace Jupiter.Implementation
 				LogReplicationHeartbeat();
 
 				await Parallel.ForEachAsync(GetBlobEventsAsync(ns, refBucket, replicationToken),
-					new ParallelOptions { MaxDegreeOfParallelism = maxParallelism, CancellationToken = linkedTokenSource.Token },
+					new ParallelOptions { MaxDegreeOfParallelism = maxParallelism, CancellationToken = replicationToken },
 					async (BlobReplicationLogEvent @event, CancellationToken ctx) =>
 				{
+					if (ctx.IsCancellationRequested)
+						return;
+
 					using TelemetrySpan scope = _tracer.StartActiveSpan("replicator.replicate_blob")
 						.SetAttribute("operation.name", "replicator.replicate_blob")
 						.SetAttribute("resource.name", $"[REDACTED_EMAIL]}")
@@ -293,7 +293,7 @@ namespace Jupiter.Implementation
 						// we do not need to replicate delete events
 						if (@event.Op != BlobReplicationLogEvent.OpType.Deleted)
 						{
-							blobWasReplicated = await ReplicateBlobAsync(@event.Namespace, @event.Blob, @event.BucketHint, linkedTokenSource.Token);
+							blobWasReplicated = await ReplicateBlobAsync(@event.Namespace, @event.Blob, @event.BucketHint, ctx);
 						}
 
 						if (blobWasReplicated)
@@ -310,6 +310,8 @@ namespace Jupiter.Implementation
 				// finished replicating bucket
 				DateTime timestamp = refBucket.FromReplicationBucketIdentifier();
 
+				_logger.LogDebug("{Name} finished enumeration of bucket: {EventBucket}", _name, refBucket);
+
 				// if the bucket was recently created we do not store that we replicate it, thus we will replicate it again to avoid risks with missing recent blobs that gets added after we queries the time bucket
 				bool isRecentBucket = timestamp > DateTime.UtcNow.AddMinutes(-10);
 
