Cloud DDC - Reimplemented retry of aggregate exception during blob replication to try and fix issues with socket issues not being retries as expected

diff --git a/Engine/Source/Programs/UnrealCloudDDC/Jupiter/Implementation/Replication/BlobsReplicator.cs b/Engine/Source/Programs/UnrealCloudDDC/Jupiter/Implementation/Replication/BlobsReplicator.cs
--- a/Engine/Source/Programs/UnrealCloudDDC/Jupiter/Implementation/Replication/BlobsReplicator.cs
+++ b/Engine/Source/Programs/UnrealCloudDDC/Jupiter/Implementation/Replication/BlobsReplicator.cs
@@ -536,16 +536,20 @@ namespace Jupiter.Implementation
 					await _blobService.PutObjectAsync(ns, payload, blob, bucketHint: bucketHint, cancellationToken: cancellationToken);
 					break;
 				}
-				catch (AggregateException e)
+				catch (AggregateException ae)
 				{
-					if (e.InnerException is not IOException && e.InnerException is not HttpIOException)
+					ae.Handle(e =>
 					{
-						_logger.LogError(e, "Failed to replicate {Blob} in {Namespace} due to unknown exception", blob, ns);
-						throw;
-					}
-
-					_logger.LogWarning("Failed to replicate {Blob} in {Namespace} due to retriable exception {Exception}", blob, ns, e);
-					lastException = e;
+						if (e is not IOException or HttpIOException)
+						{
+							_logger.LogError(e, "Failed to replicate {Blob} in {Namespace} due to unknown exception", blob, ns);
+							return false;
+						}
+						
+						_logger.LogWarning("Failed to replicate {Blob} in {Namespace} due to retriable exception {Exception}", blob, ns, e);
+						lastException = e;
+						return true;
+					});
 				}
 				catch (HttpRequestException e)
 				{
