Cloud DDC - Add multipart support for blob replication

diff --git a/Engine/Source/Programs/UnrealCloudDDC/Helm/UnrealCloudDDC/templates/configmap-worker.yaml b/Engine/Source/Programs/UnrealCloudDDC/Helm/UnrealCloudDDC/templates/configmap-worker.yaml
--- a/Engine/Source/Programs/UnrealCloudDDC/Helm/UnrealCloudDDC/templates/configmap-worker.yaml
+++ b/Engine/Source/Programs/UnrealCloudDDC/Helm/UnrealCloudDDC/templates/configmap-worker.yaml
@@ -65,7 +65,8 @@ Replication:
 {{ range . }}
   - ReplicatorName: {{ required "ReplicatorName must be specified for a replicator" .ReplicatorName }}
     NamespaceToReplicate: {{ required "Namespace must be specified for a replicator" .Namespace }}
-    MaxParallelReplications: {{ default 128  .MaxParallelReplications }} 
+    MaxParallelReplications: {{ default 128  .MaxParallelReplications }}
+    EnableMultipartReplications: {{ default false  .EnableMultipartReplications }}
     SkipSnapshot: {{ default true  .SkipSnapshot }}
     PageSize: {{ default 1000  .PageSize }}
     Version: {{ default "Refs"  .Version }}
diff --git a/Engine/Source/Programs/UnrealCloudDDC/Helm/UnrealCloudDDC/templates/configmap.yaml b/Engine/Source/Programs/UnrealCloudDDC/Helm/UnrealCloudDDC/templates/configmap.yaml
--- a/Engine/Source/Programs/UnrealCloudDDC/Helm/UnrealCloudDDC/templates/configmap.yaml
+++ b/Engine/Source/Programs/UnrealCloudDDC/Helm/UnrealCloudDDC/templates/configmap.yaml
@@ -106,6 +106,7 @@ Replication:
   - ReplicatorName: {{ required "ReplicatorName must be specified for a replicator" .ReplicatorName }}
     NamespaceToReplicate: {{ required "Namespace must be specified for a replicator" .Namespace }}
     MaxParallelReplications: {{ default 64  .MaxParallelReplications }}
+    EnableMultipartReplications: {{ default false  .EnableMultipartReplications }}
     SkipSnapshot: {{ default true  .SkipSnapshot }}
     PageSize: {{ default 1000  .PageSize }}
     {{- if .SourceSite }}
diff --git a/Engine/Source/Programs/UnrealCloudDDC/Jupiter/Implementation/Replication/BlobsReplicator.cs b/Engine/Source/Programs/UnrealCloudDDC/Jupiter/Implementation/Replication/BlobsReplicator.cs
--- a/Engine/Source/Programs/UnrealCloudDDC/Jupiter/Implementation/Replication/BlobsReplicator.cs
+++ b/Engine/Source/Programs/UnrealCloudDDC/Jupiter/Implementation/Replication/BlobsReplicator.cs
@@ -7,6 +7,7 @@ using System.IO;
 using System.Linq;
 using System.Net;
 using System.Net.Http;
+using System.Net.Http.Headers;
 using System.Runtime.CompilerServices;
 using System.Text.Json;
 using System.Threading;
@@ -247,10 +248,7 @@ namespace Jupiter.Implementation
 		{
 			int countOfReplicationsDone = 0;
 
-			// if MaxParallelReplications is set to not limit we use the default behavior of ParallelForEachAsync which is to limit based on CPUs 
-			int maxParallelism = _replicatorSettings.MaxParallelReplications != -1 ? _replicatorSettings.MaxParallelReplications : 0;
-
-			_logger.LogInformation("{Name} Starting blob replication maxParallelism: {MaxParallelism} Last Bucket {LastBucket}", _name, maxParallelism, lastBucket);
+			_logger.LogInformation("{Name} Starting blob replication maxParallelism: {MaxParallelism} Last Bucket {LastBucket}", _name, _replicatorSettings.MaxParallelReplications, lastBucket);
 
 			if (replicationToken.IsCancellationRequested)
 			{
@@ -272,11 +270,13 @@ namespace Jupiter.Implementation
 				LogReplicationHeartbeat();
 
 				await Parallel.ForEachAsync(GetBlobEventsAsync(ns, refBucket, replicationToken),
-					new ParallelOptions { MaxDegreeOfParallelism = maxParallelism, CancellationToken = replicationToken },
+					new ParallelOptions { MaxDegreeOfParallelism = _replicatorSettings.MaxParallelReplications, CancellationToken = replicationToken },
 					async (BlobReplicationLogEvent @event, CancellationToken ctx) =>
 				{
 					if (ctx.IsCancellationRequested)
+					{
 						return;
+					}
 
 					using TelemetrySpan scope = _tracer.StartActiveSpan("replicator.replicate_blob")
 						.SetAttribute("operation.name", "replicator.replicate_blob")
@@ -488,10 +488,260 @@ namespace Jupiter.Implementation
 				await _blobService.ImportRegionalBlobAsync(ns, _replicatorSettings.SourceSite, blob, bucketHint, cancellationToken);
 				return true;
 			}
+
+			if (_replicatorSettings.EnableMultipartReplications && _blobService.IsMultipartUploadSupported(ns))
+			{
+				return await ReplicateMultipartBlobAsync(ns, blob, bucketHint, cancellationToken);
+			}
 			
 			return await ReplicateBlobHttpAsync(ns, blob, bucketHint, cancellationToken);
 		}
 
+		private async Task<bool> ReplicateMultipartBlobAsync(NamespaceId ns, BlobId blob, BucketId? bucketHint, CancellationToken cancellationToken)
+		{
+			const int RetryAttempts = 5;
+			Exception? lastException = null;
+			for (int i = 0; i < RetryAttempts; i++)
+			{
+				lastException = null;
+				
+				try
+				{
+					const long InitialRangeFrom = 0;
+					long initialRangeTo = _blobService.GetMultipartLimits(ns)!.IdealChunkSize - 1;
+					
+					using HttpRequestMessage blobRequest = await BuildHttpRequestAsync(HttpMethod.Get, new Uri($"api/v1/blobs/{ns}/{blob}", UriKind.Relative));
+					blobRequest.Headers.Range = new RangeHeaderValue(InitialRangeFrom, initialRangeTo);
+					
+					using HttpResponseMessage blobResponse = await _httpClient.SendAsync(blobRequest, HttpCompletionOption.ResponseHeadersRead, cancellationToken);
+
+					if (blobResponse.StatusCode == HttpStatusCode.NotFound)
+					{
+						_logger.LogWarning("Failed to replicate {Blob} in {Namespace} due to it not existing.", blob, ns);
+						return false;
+					}
+
+					// Requested blob is smaller than the initial range. Fallback to request the entire blob
+					if (blobResponse.StatusCode == HttpStatusCode.RequestedRangeNotSatisfiable)
+					{
+						return await ReplicateBlobHttpAsync(ns, blob, bucketHint, cancellationToken);
+					}
+
+					// Server didn't respect the range request for some reason, but successfully retrieved the content
+					if (blobResponse.StatusCode == HttpStatusCode.OK)
+					{
+						await using Stream s = await blobResponse.Content.ReadAsStreamAsync(cancellationToken);
+						long? contentLength = blobResponse.Content.Headers.ContentLength;
+
+						if (contentLength == null)
+						{
+							throw new Exception("Expected content-length on blob response");
+						}
+
+						using IBufferedPayload payload = await _bufferedPayloadFactory.CreateFromStreamAsync(s, contentLength.Value, "blob-replicator", cancellationToken);
+						await _blobService.PutObjectAsync(ns, payload, blob, bucketHint: bucketHint, cancellationToken: cancellationToken);
+
+						return true;
+					}
+
+					if (blobResponse.StatusCode != HttpStatusCode.PartialContent)
+					{
+						_logger.LogError("Bad http response when replicating {Blob} in {Namespace} . Status code: {StatusCode}", blob, ns, blobResponse.StatusCode);
+						return false;
+					}
+					
+					// Received the first chunk. Use response headers to determine how to parallelize the rest
+					ContentRangeHeaderValue responseRange = blobResponse.Content.Headers.ContentRange!;
+
+					if (responseRange.From != InitialRangeFrom || responseRange.To != initialRangeTo)
+					{
+						throw new Exception("Response range does not match requested range");
+					}
+
+					if (!responseRange.HasLength)
+					{
+						throw new Exception("Expected length in Content-Range header");
+					}
+					
+					(string? uploadId, string? blobName) = await _blobService.StartMultipartUploadAsync(ns, cancellationToken);
+					
+					if (uploadId == null || blobName == null)
+					{
+						throw new NotImplementedException("Failed to initialize multipart upload, no upload id returned");
+					}
+					
+					List<MultipartByteRange> ranges = _blobService.GetMultipartRanges(ns, blobName, uploadId, (ulong)responseRange.Length!, cancellationToken);
+					string firstPartId = ranges.Single(range => range.FirstByte == 0).PartId;
+					List<MultipartByteRange> remainingRanges = ranges.Where(range => range.PartId != firstPartId).ToList();
+					
+					{
+						await using Stream s = await blobResponse.Content.ReadAsStreamAsync(cancellationToken);
+						long? contentLength = blobResponse.Content.Headers.ContentLength;
+
+						if (contentLength == null)
+						{
+							throw new Exception("Expected content-length on blob response");
+						}
+
+						using IBufferedPayload payload = await _bufferedPayloadFactory.CreateFromStreamAsync(s, contentLength.Value, "blob-replicator", cancellationToken);
+						byte[] payloadData = await s.ReadAllBytesAsync(cancellationToken);
+						await _blobService.PutMultipartUploadAsync(ns, blobName, uploadId, firstPartId, payloadData, cancellationToken);
+					}
+					
+					await Parallel.ForEachAsync(remainingRanges, cancellationToken,
+						async (range, token) => await ReplicateMultipartBlobRangeHttpAsync(ns, blob, blobName, uploadId, range, token));
+					
+					try
+					{
+						await _blobService.CompleteMultipartUploadAsync(ns, blobName, uploadId, ranges.Select(range => range.PartId).ToList(), cancellationToken);
+					}
+					catch (MissingMultipartPartsException e)
+					{
+						_logger.LogError(e, "Failed to replicate all parts of multipart replication {Blob} in {Namespace}", blob, ns);
+						return false;
+					}
+
+					try
+					{
+						await _blobService.VerifyMultipartUpload(ns, blob, blobName, isCompressed: false, cancellationToken: cancellationToken);
+					}
+					catch (HashMismatchException e)
+					{
+						_logger.LogError(e, "Hash mismatch for multipart replication {Blob} in {Namespace}", blob, ns);
+						return false;
+					}
+
+					break;
+				}
+				catch (AggregateException ae)
+				{
+					ae.Handle(e =>
+					{
+						if (e is not IOException or HttpIOException)
+						{
+							_logger.LogError(e, "Failed to replicate {Blob} in {Namespace} due to unknown exception", blob, ns);
+							return false;
+						}
+						
+						_logger.LogWarning("Failed to replicate {Blob} in {Namespace} due to retriable exception {Exception}", blob, ns, e);
+						lastException = e;
+						return true;
+					});
+				}
+				catch (HttpRequestException e)
+				{
+					// rethrow unknown exceptions
+					if (e.InnerException is not IOException && e.InnerException is not HttpIOException)
+					{
+						_logger.LogError(e, "Failed to replicate {Blob} in {Namespace} due to unknown exception", blob, ns);
+						throw;
+					}
+
+					// is an io exception that we should retry
+					lastException = e;
+
+					if (e.StatusCode == HttpStatusCode.NotFound)
+					{
+						// delay the retry with a second if blob is missing, as its really not expected for it to actually miss, likely we just haven't reached consistency on the fact that it does exist
+						await Task.Delay(1000, cancellationToken);
+					}
+
+					_logger.LogWarning("Failed to replicate {Blob} in {Namespace} due to retriable exception {Exception}.", blob, ns, e);
+				}
+			}
+
+			if (lastException != null)
+			{
+				throw lastException;
+			}
+
+			return true;
+		}
+
+		private async Task ReplicateMultipartBlobRangeHttpAsync(NamespaceId ns, BlobId blob, string blobName, string uploadId, MultipartByteRange range, CancellationToken cancellationToken)
+		{
+			_logger.LogDebug("Retrieving blob {Blob} chunk range {Start}-{End} in {Namespace}.", blob, range.FirstByte, range.LastByte, ns);
+
+			const int RetryAttempts = 5;
+			Exception? lastException = null;
+			for (int i = 0; i < RetryAttempts; i++)
+			{
+				lastException = null;
+
+				try
+				{
+					using HttpRequestMessage blobRequest = await BuildHttpRequestAsync(HttpMethod.Get, new Uri($"api/v1/blobs/{ns}/{blob}", UriKind.Relative));
+					blobRequest.Headers.Range = new RangeHeaderValue((long)range.FirstByte, (long)range.LastByte);
+
+					using HttpResponseMessage blobResponse = await _httpClient.SendAsync(blobRequest, HttpCompletionOption.ResponseHeadersRead, cancellationToken);
+
+					if (blobResponse.StatusCode != HttpStatusCode.PartialContent)
+					{
+						_logger.LogError("Bad http response when retrieving {Blob} chunk range {Start}-{End} in {Namespace} . Status code: {StatusCode}", blob, range.FirstByte, range.LastByte, ns, blobResponse.StatusCode);
+					}
+					
+					ContentRangeHeaderValue responseRange = blobResponse.Content.Headers.ContentRange!;
+
+					if (responseRange.From != (long)range.FirstByte || responseRange.To != (long)range.LastByte)
+					{
+						throw new Exception("Response range does not match requested range");
+					}
+
+					await using Stream s = await blobResponse.Content.ReadAsStreamAsync(cancellationToken);
+					long? contentLength = blobResponse.Content.Headers.ContentLength;
+
+					if (contentLength == null)
+					{
+						throw new Exception("Expected content-length on blob response");
+					}
+
+					using IBufferedPayload payload = await _bufferedPayloadFactory.CreateFromStreamAsync(s, contentLength.Value, "blob-replicator", cancellationToken);
+					byte[] payloadData = await s.ReadAllBytesAsync(cancellationToken);
+					await _blobService.PutMultipartUploadAsync(ns, blobName, uploadId, range.PartId, payloadData, cancellationToken);
+				}
+				catch (AggregateException ae)
+				{
+					ae.Handle(e =>
+					{
+						if (e is not IOException or HttpIOException)
+						{
+							_logger.LogError(e, "Failed to replicate {Blob} in {Namespace} due to unknown exception", blob, ns);
+							return false;
+						}
+						
+						_logger.LogWarning("Failed to replicate {Blob} in {Namespace} due to retriable exception {Exception}", blob, ns, e);
+						lastException = e;
+						return true;
+					});
+				}
+				catch (HttpRequestException e)
+				{
+					// rethrow unknown exceptions
+					if (e.InnerException is not IOException && e.InnerException is not HttpIOException)
+					{
+						_logger.LogError(e, "Failed to replicate {Blob} in {Namespace} due to unknown exception", blob, ns);
+						throw;
+					}
+
+					// is an io exception that we should retry
+					lastException = e;
+
+					if (e.StatusCode == HttpStatusCode.NotFound)
+					{
+						// delay the retry with a second if blob is missing, as its really not expected for it to actually miss, likely we just haven't reached consistency on the fact that it does exist
+						await Task.Delay(1000, cancellationToken);
+					}
+
+					_logger.LogWarning("Failed to replicate {Blob} in {Namespace} due to retriable exception {Exception}.", blob, ns, e);
+				}
+			}
+
+			if (lastException != null)
+			{
+				throw lastException;
+			}
+		}
+
 		private async Task<bool> ReplicateBlobHttpAsync(NamespaceId ns, BlobId blob, BucketId? bucketHint, CancellationToken cancellationToken)
 		{
 			const int RetryAttempts = 5;
diff --git a/Engine/Source/Programs/UnrealCloudDDC/Jupiter/Implementation/Replication/ReplicationSettings.cs b/Engine/Source/Programs/UnrealCloudDDC/Jupiter/Implementation/Replication/ReplicationSettings.cs
--- a/Engine/Source/Programs/UnrealCloudDDC/Jupiter/Implementation/Replication/ReplicationSettings.cs
+++ b/Engine/Source/Programs/UnrealCloudDDC/Jupiter/Implementation/Replication/ReplicationSettings.cs
@@ -73,6 +73,11 @@ namespace Jupiter
 		/// </summary>
 		public int MaxParallelReplications { get; set; } = 64;
 
+		/// <summary>
+		/// Whether to enable multipart transfers, when using blob stores that support it.
+		/// </summary>
+		public bool EnableMultipartReplications { get; set; } = false;
+
 		/// <summary>
 		/// Do not read in old history using a snapshot, instead just start replicating the new incremental state. Means we will only have partial state but we will also get going quicker on relatively useful blobs.
 		/// </summary>
