Cloud DDC - Return 503 if we are unable to reach scylla when fetching a compressed buffer blob

diff --git a/Engine/Source/Programs/UnrealCloudDDC/Jupiter/Controllers/CompressedBlobController.cs b/Engine/Source/Programs/UnrealCloudDDC/Jupiter/Controllers/CompressedBlobController.cs
--- a/Engine/Source/Programs/UnrealCloudDDC/Jupiter/Controllers/CompressedBlobController.cs
+++ b/Engine/Source/Programs/UnrealCloudDDC/Jupiter/Controllers/CompressedBlobController.cs
@@ -103,6 +103,10 @@ namespace Jupiter.Controllers
 			{
 				return NotFound(new ValidationProblemDetails { Title = $"Content Id {e.ContentId} not found" });
 			}
+			catch (RemoteServiceException e)
+			{
+				return Problem(e.Message, null, (int)HttpStatusCode.ServiceUnavailable);
+			}
 		}
 
 		[HttpHead("{ns}/{id}")]
