Cloud DDC - Fixed issue with the audit log null access when no part and no build id was specified but a namespace was.

diff --git a/Engine/Source/Programs/UnrealCloudDDC/Jupiter/Controllers/BuildsController.cs b/Engine/Source/Programs/UnrealCloudDDC/Jupiter/Controllers/BuildsController.cs
--- a/Engine/Source/Programs/UnrealCloudDDC/Jupiter/Controllers/BuildsController.cs
+++ b/Engine/Source/Programs/UnrealCloudDDC/Jupiter/Controllers/BuildsController.cs
@@ -547,7 +547,7 @@ namespace Jupiter.Controllers
 			{
 				await Parallel.ForEachAsync(bucketsToSearch, cancellationTokenSource.Token, async (bucket, token) =>
 				{
-					LogAuditEntry(HttpMethod.Post, nameof(SearchBuildAsync), ns, bucket, null, null);
+					LogAuditEntry(HttpMethod.Post, nameof(SearchBucketsAsync), ns, bucket, null, null);
 				
 					await foreach (BuildMetadata build in _buildStore.ListBuildsAsync(ns, bucket, includeTTL).WithCancellation(token))
 					{
@@ -1932,9 +1932,9 @@ namespace Jupiter.Controllers
 				return;
 			}
 
-			if (partId == null)
+			if (partId == null || buildId == null)
 			{
-				_auditLogger?.Information("{HttpMethod} {Endpoint} '{Namespace}/{BucketId}:{BuildId}' IP:{IP} User:{Username} UserAgent:\"{Useragent}\" SessionId: {SessionId} Jti:{Jti}", method, endpoint, ns, bucketId, buildId, Request.HttpContext.Connection.RemoteIpAddress, User?.Identity?.Name ?? "Unknown-user", string.Join(' ', Request.Headers.UserAgent.ToArray()), sessionHeader, jti);
+				_auditLogger?.Information("{HttpMethod} {Endpoint} '{Namespace}/{BucketId}' IP:{IP} User:{Username} UserAgent:\"{Useragent}\" SessionId: {SessionId} Jti:{Jti}", method, endpoint, ns, bucketId, Request.HttpContext.Connection.RemoteIpAddress, User?.Identity?.Name ?? "Unknown-user", string.Join(' ', Request.Headers.UserAgent.ToArray()), sessionHeader, jti);
 				return;
 			}
 			
