Cloud DDC - Removed redundant blob replication trace

diff --git a/Engine/Source/Programs/UnrealCloudDDC/Jupiter/Implementation/Replication/BlobsReplicator.cs b/Engine/Source/Programs/UnrealCloudDDC/Jupiter/Implementation/Replication/BlobsReplicator.cs
--- a/Engine/Source/Programs/UnrealCloudDDC/Jupiter/Implementation/Replication/BlobsReplicator.cs
+++ b/Engine/Source/Programs/UnrealCloudDDC/Jupiter/Implementation/Replication/BlobsReplicator.cs
@@ -469,10 +469,6 @@ namespace Jupiter.Implementation
 
 		private async Task<bool> ReplicateBlobAsync(NamespaceId ns, BlobId blob, BucketId? bucketHint, CancellationToken cancellationToken)
 		{
-			using TelemetrySpan scope = _tracer.StartActiveSpan("replicator.replicate_blob")
-				.SetAttribute("operation.name", "replicator.replicate_blob")
-				.SetAttribute("resource.name", $"{ns}.{blob}");
-
 			// the point of the replicator is to transfer the blobs locally, thus we ignore if they exist in remote locations
 			bool exists = await _blobService.ExistsAsync(ns, blob, storageLayers: null, ignoreRemoteBlobs: true, cancellationToken: cancellationToken);
 			if (exists)
