Cloud DDC - Added more logging for null dereference around audit logs, our line numbers from the stacktrace are unreliable so adding more logging to try and catch what is wrong

diff --git a/Engine/Source/Programs/UnrealCloudDDC/Jupiter/Controllers/BuildsController.cs b/Engine/Source/Programs/UnrealCloudDDC/Jupiter/Controllers/BuildsController.cs
--- a/Engine/Source/Programs/UnrealCloudDDC/Jupiter/Controllers/BuildsController.cs
+++ b/Engine/Source/Programs/UnrealCloudDDC/Jupiter/Controllers/BuildsController.cs
@@ -1913,32 +1913,42 @@ namespace Jupiter.Controllers
 		private void LogAuditEntry(HttpMethod method, string endpoint, NamespaceId? ns, BucketId? bucketId, CbObjectId? buildId, string? partId)
 		{
 			string sessionHeader = "None";
-			if (Request.Headers.TryGetValue("ue-session", out StringValues sessionHeaderValues))
+			if (Request?.Headers == null)
 			{
-				sessionHeader = sessionHeaderValues.First() ?? "None";
+				// we have seen null dereference when processing the request, adding this logging to try and track down the cause
+				_logger.LogWarning("Request was null");
+			}
+
+			if (Request!.Headers.TryGetValue("ue-session", out StringValues sessionHeaderValues))
+			{
+				sessionHeader = sessionHeaderValues.FirstOrDefault() ?? "None";
 			}
 
 			string? jti = User.Claims.FirstOrDefault(claim => claim.Type == "jti")?.Value ?? null;
 
+			string userAgentString = string.Join(' ', Request.Headers.UserAgent.ToArray());
+			string username = User.Identity?.Name ?? "Unknown-user";
+			string ip = Request.HttpContext.Connection.RemoteIpAddress?.ToString() ?? "Unknown";
+
 			if (ns == null)
 			{
-				_auditLogger?.Information("{HttpMethod} {Endpoint} IP:{IP} User:{Username} UserAgent:\"{Useragent}\" SessionId: {SessionId} Jti:{Jti}", method, endpoint, Request.HttpContext.Connection.RemoteIpAddress, User?.Identity?.Name ?? "Unknown-user", string.Join(' ', Request.Headers.UserAgent.ToArray()), sessionHeader, jti);
+				_auditLogger?.Information("{HttpMethod} {Endpoint} IP:{IP} User:{Username} UserAgent:\"{Useragent}\" SessionId: {SessionId} Jti:{Jti}", method, endpoint, ip, username, userAgentString, sessionHeader, jti);
 				return;
 			}
 
 			if (bucketId == null)
 			{
-				_auditLogger?.Information("{HttpMethod} {Endpoint} '{Namespace}' IP:{IP} User:{Username} UserAgent:\"{Useragent}\" SessionId: {SessionId} Jti:{Jti}", method, endpoint, ns, Request.HttpContext.Connection.RemoteIpAddress, User?.Identity?.Name ?? "Unknown-user", string.Join(' ', Request.Headers.UserAgent.ToArray()), sessionHeader, jti);
+				_auditLogger?.Information("{HttpMethod} {Endpoint} '{Namespace}' IP:{IP} User:{Username} UserAgent:\"{Useragent}\" SessionId: {SessionId} Jti:{Jti}", method, endpoint, ns, ip, username, userAgentString, sessionHeader, jti);
 				return;
 			}
 
 			if (partId == null || buildId == null)
 			{
-				_auditLogger?.Information("{HttpMethod} {Endpoint} '{Namespace}/{BucketId}' IP:{IP} User:{Username} UserAgent:\"{Useragent}\" SessionId: {SessionId} Jti:{Jti}", method, endpoint, ns, bucketId, Request.HttpContext.Connection.RemoteIpAddress, User?.Identity?.Name ?? "Unknown-user", string.Join(' ', Request.Headers.UserAgent.ToArray()), sessionHeader, jti);
+				_auditLogger?.Information("{HttpMethod} {Endpoint} '{Namespace}/{BucketId}' IP:{IP} User:{Username} UserAgent:\"{Useragent}\" SessionId: {SessionId} Jti:{Jti}", method, endpoint, ns, bucketId, ip, username, userAgentString, sessionHeader, jti);
 				return;
 			}
 			
-			_auditLogger?.Information("{HttpMethod} {Endpoint} '{Namespace}/{BucketId}:{BuildId}' Part:{PartId} IP:{IP} User:{Username} UserAgent:\"{Useragent}\" SessionId: {SessionId} Jti:{Jti}", method, endpoint, ns, bucketId, buildId, partId, Request.HttpContext.Connection.RemoteIpAddress, User?.Identity?.Name ?? "Unknown-user", string.Join(' ', Request.Headers.UserAgent.ToArray()), sessionHeader, jti);
+			_auditLogger?.Information("{HttpMethod} {Endpoint} '{Namespace}/{BucketId}:{BuildId}' Part:{PartId} IP:{IP} User:{Username} UserAgent:\"{Useragent}\" SessionId: {SessionId} Jti:{Jti}", method, endpoint, ns, bucketId, buildId, partId, ip, username, userAgentString, sessionHeader, jti);
 		}
 	}
 
