Cloud DDC - Use the requested aborted cancellation token when we are sending blobs using the compressed-blob endpoint

diff --git a/Engine/Source/Programs/UnrealCloudDDC/Jupiter/Controllers/CompressedBlobController.cs b/Engine/Source/Programs/UnrealCloudDDC/Jupiter/Controllers/CompressedBlobController.cs
--- a/Engine/Source/Programs/UnrealCloudDDC/Jupiter/Controllers/CompressedBlobController.cs
+++ b/Engine/Source/Programs/UnrealCloudDDC/Jupiter/Controllers/CompressedBlobController.cs
@@ -69,7 +69,7 @@ namespace Jupiter.Controllers
 
 			try
 			{
-				(BlobContents blobContents, string mediaType, BlobId? contentHash) = await _storage.GetCompressedObjectAsync(ns, id, HttpContext.RequestServices, supportsRedirectUri: supportsRedirect);
+				(BlobContents blobContents, string mediaType, BlobId? contentHash) = await _storage.GetCompressedObjectAsync(ns, id, HttpContext.RequestServices, supportsRedirectUri: supportsRedirect, cancellationToken: HttpContext.RequestAborted);
 
 				StringValues acceptHeader = Request.Headers["Accept"];
 				if (!acceptHeader.Contains("*/*") && acceptHeader.Count != 0 && !acceptHeader.Contains(mediaType))
@@ -121,7 +121,7 @@ namespace Jupiter.Controllers
 				return result;
 			}
 
-			BlobId[]? chunks = await _contentIdStore.ResolveAsync(ns, id, mustBeContentId: false);
+			BlobId[]? chunks = await _contentIdStore.ResolveAsync(ns, id, mustBeContentId: false, cancellationToken: HttpContext.RequestAborted);
 			if (chunks == null || chunks.Length == 0)
 			{
 				return NotFound();
