Cloud DDC - Fix Helm ConfigMap generation when omitting any Jupiter configuration

diff --git a/Engine/Source/Programs/UnrealCloudDDC/Helm/UnrealCloudDDC/templates/configmap.yaml b/Engine/Source/Programs/UnrealCloudDDC/Helm/UnrealCloudDDC/templates/configmap.yaml
--- a/Engine/Source/Programs/UnrealCloudDDC/Helm/UnrealCloudDDC/templates/configmap.yaml
+++ b/Engine/Source/Programs/UnrealCloudDDC/Helm/UnrealCloudDDC/templates/configmap.yaml
@@ -80,7 +80,9 @@ UnrealCloudDDC:
 
 {{- if or ($config.Jupiter) (.Values.nginx.useDomainSockets) }}
 Jupiter: 
+{{ if $config.Jupiter }}
 {{ $config.Jupiter | toYaml | nindent 2 }}
+{{- end }}
 {{ if .Values.nginx.useDomainSockets }}
   UseDomainSockets: true
   DomainSocketsRoot: "/nginx"
