Cloud DDC - Add service to periodically clean up temp files

diff --git a/Engine/Source/Programs/UnrealCloudDDC/Jupiter/Implementation/GC/TempCleanupService.cs b/Engine/Source/Programs/UnrealCloudDDC/Jupiter/Implementation/GC/TempCleanupService.cs
new file mode 100644
--- /dev/null
+++ b/Engine/Source/Programs/UnrealCloudDDC/Jupiter/Implementation/GC/TempCleanupService.cs
@@ -0,0 +1,104 @@
+﻿// Copyright Epic Games, Inc. All Rights Reserved.
+
+using System;
+using System.Collections.Generic;
+using System.IO;
+using System.Linq;
+using System.Threading;
+using System.Threading.Tasks;
+using Jupiter.Common.Implementation;
+using Microsoft.Extensions.Logging;
+using Microsoft.Extensions.Options;
+
+namespace Jupiter.Implementation
+{
+	public class TempCleanupService : PollingService<object>
+	{
+		private readonly IOptionsMonitor<GCSettings> _gcSettings;
+		private readonly IOptionsMonitor<BufferedPayloadOptions> _bufferedPayloadOptions;
+		private volatile bool _alreadyPolling;
+		private readonly ILogger _logger;
+
+		public TempCleanupService(IOptionsMonitor<GCSettings> gcSettings, IOptionsMonitor<BufferedPayloadOptions> bufferedPayloadOptions, ILogger<TempCleanupService> logger) : base(serviceName: nameof(TempCleanupService), gcSettings.CurrentValue.TempCleanupPollFrequency, new object(), logger)
+		{
+			_gcSettings = gcSettings;
+			_bufferedPayloadOptions = bufferedPayloadOptions;
+			_logger = logger;
+		}
+
+		protected override bool ShouldStartPolling()
+		{
+			return _gcSettings.CurrentValue.CleanTempFiles;
+		}
+
+		public override async Task<bool> OnPollAsync(object state, CancellationToken cancellationToken)
+		{
+			if (_alreadyPolling)
+			{
+				return false;
+			}
+
+			_alreadyPolling = true;
+			try
+			{
+				await Task.Run(() => Cleanup(cancellationToken), cancellationToken);
+				return true;
+			}
+			catch (Exception ex)
+			{
+				_logger.LogError(ex, "Exception running temp file cleanup");
+				return false;
+			}
+			finally
+			{
+				_alreadyPolling = false;
+			}
+		}
+
+		private void Cleanup(CancellationToken cancellationToken)
+		{
+			_logger.LogInformation("Beginning temp file cleanup");
+			
+			IEnumerable<FileInfo> fileInfos = GetObjectsOlderThan(DateTime.Now - _gcSettings.CurrentValue.TempCleanupMaximumAge);
+
+			ulong countOfBlobsRemoved = 0;
+			
+			foreach (FileInfo fi in fileInfos)
+			{
+				if (cancellationToken.IsCancellationRequested)
+				{
+					break;
+				}
+				
+				try
+				{
+					fi.Delete();
+					countOfBlobsRemoved++;
+				}
+				catch (FileNotFoundException)
+				{
+					// Already deleted
+				}
+				catch (DirectoryNotFoundException)
+				{
+					// Already deleted
+				}
+			}
+			
+			_logger.LogInformation("Removed {Count} temp files", countOfBlobsRemoved);
+		}
+
+		private IEnumerable<FileInfo> GetObjectsOlderThan(DateTime cutoff)
+		{
+			DirectoryInfo di = new DirectoryInfo(_bufferedPayloadOptions.CurrentValue.FilesystemTempPayloadRoot);
+			
+			if (!di.Exists)
+			{
+				return Array.Empty<FileInfo>();
+			}
+			
+			return di.EnumerateFiles("*", SearchOption.AllDirectories).Where(x => x.LastWriteTime < cutoff);
+		}
+	}
+}
+
diff --git a/Engine/Source/Programs/UnrealCloudDDC/Jupiter/Jupiter.Startup.cs b/Engine/Source/Programs/UnrealCloudDDC/Jupiter/Jupiter.Startup.cs
--- a/Engine/Source/Programs/UnrealCloudDDC/Jupiter/Jupiter.Startup.cs
+++ b/Engine/Source/Programs/UnrealCloudDDC/Jupiter/Jupiter.Startup.cs
@@ -206,6 +206,9 @@ namespace Jupiter
 			services.AddSingleton<RefCleanupService>();
 			services.AddHostedService<RefCleanupService>(p => p.GetService<RefCleanupService>()!);
 
+			services.AddSingleton<TempCleanupService>();
+			services.AddHostedService<TempCleanupService>(p => p.GetService<TempCleanupService>()!);
+
 			// This will technically create a new settings object, but as we do not dynamically update it the settings will reflect what we actually want
 			UnrealCloudDDCSettings settings = services.BuildServiceProvider().GetService<IOptionsMonitor<UnrealCloudDDCSettings>>()!.CurrentValue!;
 
diff --git a/Engine/Source/Programs/UnrealCloudDDC/Jupiter/Settings.cs b/Engine/Source/Programs/UnrealCloudDDC/Jupiter/Settings.cs
--- a/Engine/Source/Programs/UnrealCloudDDC/Jupiter/Settings.cs
+++ b/Engine/Source/Programs/UnrealCloudDDC/Jupiter/Settings.cs
@@ -256,12 +256,15 @@ namespace Jupiter
 
 		public bool CleanOldRefRecords { get; set; } = false;
 		public bool CleanOldBlobs { get; set; } = true;
+		public bool CleanTempFiles { get; set; } = true;
 		public bool RunFilesystemCleanup { get; set; } = false;
 
 		public TimeSpan LastAccessCutoff { get; set; } = TimeSpan.FromDays(14);
 
 		public TimeSpan BlobCleanupPollFrequency { get; set; } = TimeSpan.FromMinutes(60);
 		public TimeSpan RefCleanupPollFrequency { get; set; } = TimeSpan.FromMinutes(60);
+		public TimeSpan TempCleanupPollFrequency { get; set; } = TimeSpan.FromMinutes(60);
+		public TimeSpan TempCleanupMaximumAge { get; set; } = TimeSpan.FromHours(6);
 		public int OrphanGCMaxParallelOperations { get; set; } = 8;
 		public int OrphanRefMaxParallelOperations { get; set; } = 8;
 		public bool WriteDeleteToReplicationLog { get; set; } = false;
