Cloud DDC - Added overload for LogAuditEntry that does not rely on the current controller state and used that in the parallel callbacks for searching as this was getting null values in a few requests but was succesful for other buckets within the same request.

diff --git a/Engine/Source/Programs/UnrealCloudDDC/Jupiter/Controllers/BuildsController.cs b/Engine/Source/Programs/UnrealCloudDDC/Jupiter/Controllers/BuildsController.cs
--- a/Engine/Source/Programs/UnrealCloudDDC/Jupiter/Controllers/BuildsController.cs
+++ b/Engine/Source/Programs/UnrealCloudDDC/Jupiter/Controllers/BuildsController.cs
@@ -545,9 +545,21 @@ namespace Jupiter.Controllers
 			using CancellationTokenSource cancellationTokenSource = new CancellationTokenSource();
 			try
 			{
+				string sessionHeader = "None";
+
+				if (Request!.Headers.TryGetValue("ue-session", out StringValues sessionHeaderValues))
+				{
+					sessionHeader = sessionHeaderValues.FirstOrDefault() ?? "None";
+				}
+
+				string jti = User.Claims.FirstOrDefault(claim => claim.Type == "jti")?.Value ?? "<unknown>";
+				string userAgentString = string.Join(' ', Request.Headers.UserAgent.ToArray());
+				string username = User.Identity?.Name ?? "Unknown-user";
+				string ip = Request.HttpContext.Connection.RemoteIpAddress?.ToString() ?? "Unknown";
+
 				await Parallel.ForEachAsync(bucketsToSearch, cancellationTokenSource.Token, async (bucket, token) =>
 				{
-					LogAuditEntry(HttpMethod.Post, nameof(SearchBucketsAsync), ns, bucket, null, null);
+					LogAuditEntry(HttpMethod.Post, nameof(SearchBucketsAsync), sessionHeader, username, userAgentString, ip, jti, ns, bucket, null, null);
 				
 					await foreach (BuildMetadata build in _buildStore.ListBuildsAsync(ns, bucket, includeTTL).WithCancellation(token))
 					{
@@ -1924,31 +1936,36 @@ namespace Jupiter.Controllers
 				sessionHeader = sessionHeaderValues.FirstOrDefault() ?? "None";
 			}
 
-			string? jti = User.Claims.FirstOrDefault(claim => claim.Type == "jti")?.Value ?? null;
+			string jti = User.Claims.FirstOrDefault(claim => claim.Type == "jti")?.Value ?? "<unknown>";
 
 			string userAgentString = string.Join(' ', Request.Headers.UserAgent.ToArray());
 			string username = User.Identity?.Name ?? "Unknown-user";
 			string ip = Request.HttpContext.Connection.RemoteIpAddress?.ToString() ?? "Unknown";
 
+			LogAuditEntry(method, endpoint, sessionHeader, username, userAgentString, ip, jti, ns, bucketId, buildId, partId);
+		}
+
+		private void LogAuditEntry(HttpMethod method, string endpoint, string sessionHeader, string username, string userAgent, string ip, string jti, NamespaceId? ns, BucketId? bucketId, CbObjectId? buildId, string? partId)
+		{
 			if (ns == null)
 			{
-				_auditLogger?.Information("{HttpMethod} {Endpoint} IP:{IP} User:{Username} UserAgent:\"{Useragent}\" SessionId: {SessionId} Jti:{Jti}", method, endpoint, ip, username, userAgentString, sessionHeader, jti);
+				_auditLogger?.Information("{HttpMethod} {Endpoint} IP:{IP} User:{Username} UserAgent:\"{Useragent}\" SessionId: {SessionId} Jti:{Jti}", method, endpoint, ip, username, userAgent, sessionHeader, jti);
 				return;
 			}
 
 			if (bucketId == null)
 			{
-				_auditLogger?.Information("{HttpMethod} {Endpoint} '{Namespace}' IP:{IP} User:{Username} UserAgent:\"{Useragent}\" SessionId: {SessionId} Jti:{Jti}", method, endpoint, ns, ip, username, userAgentString, sessionHeader, jti);
+				_auditLogger?.Information("{HttpMethod} {Endpoint} '{Namespace}' IP:{IP} User:{Username} UserAgent:\"{Useragent}\" SessionId: {SessionId} Jti:{Jti}", method, endpoint, ns, ip, username, userAgent, sessionHeader, jti);
 				return;
 			}
 
 			if (partId == null || buildId == null)
 			{
-				_auditLogger?.Information("{HttpMethod} {Endpoint} '{Namespace}/{BucketId}' IP:{IP} User:{Username} UserAgent:\"{Useragent}\" SessionId: {SessionId} Jti:{Jti}", method, endpoint, ns, bucketId, ip, username, userAgentString, sessionHeader, jti);
+				_auditLogger?.Information("{HttpMethod} {Endpoint} '{Namespace}/{BucketId}' IP:{IP} User:{Username} UserAgent:\"{Useragent}\" SessionId: {SessionId} Jti:{Jti}", method, endpoint, ns, bucketId, ip, username, userAgent, sessionHeader, jti);
 				return;
 			}
 			
-			_auditLogger?.Information("{HttpMethod} {Endpoint} '{Namespace}/{BucketId}:{BuildId}' Part:{PartId} IP:{IP} User:{Username} UserAgent:\"{Useragent}\" SessionId: {SessionId} Jti:{Jti}", method, endpoint, ns, bucketId, buildId, partId, ip, username, userAgentString, sessionHeader, jti);
+			_auditLogger?.Information("{HttpMethod} {Endpoint} '{Namespace}/{BucketId}:{BuildId}' Part:{PartId} IP:{IP} User:{Username} UserAgent:\"{Useragent}\" SessionId: {SessionId} Jti:{Jti}", method, endpoint, ns, bucketId, buildId, partId, ip, username, userAgent, sessionHeader, jti);
 		}
 	}
 
