Cloud DDC - Default EnableForceSubmitRefBlobToBlobStore to false. Enabling this flag will use the old behavior, but is not recommended.

diff --git a/Engine/Source/Programs/UnrealCloudDDC/Jupiter/Settings.cs b/Engine/Source/Programs/UnrealCloudDDC/Jupiter/Settings.cs
--- a/Engine/Source/Programs/UnrealCloudDDC/Jupiter/Settings.cs
+++ b/Engine/Source/Programs/UnrealCloudDDC/Jupiter/Settings.cs
@@ -128,7 +128,7 @@ namespace Jupiter
 		/// <summary>
 		/// Forces the inlined blobs to also be submitted into the blob store, is the old behavior and is not recommended.
 		/// </summary>
-		public bool EnableForceSubmitRefBlobToBlobStore { get; set; } = true;
+		public bool EnableForceSubmitRefBlobToBlobStore { get; set; } = false;
 		public bool RequirePrivatePortForEnumeration { get; set; } = true;
 
 		public long InlineBlobMaxSize { get; set; } = 32 * 1024; // default to 32 kb blobs max
