Cloud DDC - Disable temp file cleanup for tests

#DisabledValidations Json

diff --git a/Engine/Source/Programs/UnrealCloudDDC/Jupiter/appsettings.Testing.json b/Engine/Source/Programs/UnrealCloudDDC/Jupiter/appsettings.Testing.json
--- a/Engine/Source/Programs/UnrealCloudDDC/Jupiter/appsettings.Testing.json
+++ b/Engine/Source/Programs/UnrealCloudDDC/Jupiter/appsettings.Testing.json
@@ -35,7 +35,8 @@
 	"GC": {
 		"BlobCleanupServiceEnabled": false,
 		"CleanOldBlobs": false,
-		"CleanOldRefRecords": false
+		"CleanOldRefRecords": false,
+		"CleanTempFiles": false
 	},
 
 	"S3": {
