Cloud DDC - Disable EnableBucketStatsTracking by default as it can lead to performance issues when querying sizes for many attachments

diff --git a/Engine/Source/Programs/UnrealCloudDDC/Jupiter/Settings.cs b/Engine/Source/Programs/UnrealCloudDDC/Jupiter/Settings.cs
--- a/Engine/Source/Programs/UnrealCloudDDC/Jupiter/Settings.cs
+++ b/Engine/Source/Programs/UnrealCloudDDC/Jupiter/Settings.cs
@@ -122,7 +122,7 @@ namespace Jupiter
 		public bool EnableLastAccessTracking { get; set; } = true;
 		public bool EnableOnDemandReplication { get; set; } = true;
 
-		public bool EnableBucketStatsTracking { get; set; } = true;
+		public bool EnableBucketStatsTracking { get; set; } = false;
 		public bool EnableInlineSmallBlobs { get; set; } = true;
 
 		/// <summary>
