Cloud DDC - Return 503 on RemoteServiceException in BuildsController

diff --git a/Engine/Source/Programs/UnrealCloudDDC/Jupiter/Controllers/BuildsController.cs b/Engine/Source/Programs/UnrealCloudDDC/Jupiter/Controllers/BuildsController.cs
--- a/Engine/Source/Programs/UnrealCloudDDC/Jupiter/Controllers/BuildsController.cs
+++ b/Engine/Source/Programs/UnrealCloudDDC/Jupiter/Controllers/BuildsController.cs
@@ -1623,6 +1623,10 @@ namespace Jupiter.Controllers
 			{
 				return Problem(e.Message, null, (int)HttpStatusCode.RequestTimeout);
 			}
+			catch (RemoteServiceException e)
+			{
+				return Problem(e.Message, null, (int)HttpStatusCode.ServiceUnavailable);
+			}
 		}
 
 		private ActionResult? VerifyBlockMetadata(CbObject metadataObject)
@@ -1663,41 +1667,58 @@ namespace Jupiter.Controllers
 			responseObject.BeginObject();
 			responseObject.BeginArray("blocks");
 
-			await Parallel.ForEachAsync(request.Blocks, async (blockHash, token) =>
+			try
 			{
-				// in case the block is a raw hash we need to resolve it into the actual hash of the block
-				BlobId id = BlobId.FromIoHash(blockHash);
-				BlobId[]? blobIds = await _referenceResolver.ResolveIdAsync(ns, id, HttpContext.RequestAborted);
-				if (blobIds == null)
+				await Parallel.ForEachAsync(request.Blocks, async (blockHash, token) =>
 				{
-					// if we are unable to resolve the block into an id we can not fetch it
-					return;
-				}
+					// in case the block is a raw hash we need to resolve it into the actual hash of the block
+					BlobId id = BlobId.FromIoHash(blockHash);
+					BlobId[]? blobIds = await _referenceResolver.ResolveIdAsync(ns, id, HttpContext.RequestAborted);
+					if (blobIds == null)
+					{
+						// if we are unable to resolve the block into an id we can not fetch it
+						return;
+					}
 
-				BlobId blockId = blobIds.First();
+					BlobId blockId = blobIds.First();
 
-				BlobId? metadataHash = await _blockStore.GetBlockMetadataAsync(ns, blockId);
-				if (metadataHash == null)
-				{
-					// no metadata found for this id, likely not a block
-					return;
-				}
+					BlobId? metadataHash = await _blockStore.GetBlockMetadataAsync(ns, blockId);
+					if (metadataHash == null)
+					{
+						// no metadata found for this id, likely not a block
+						return;
+					}
 
-				try
-				{
-					await using BlobContents contents = await _blobStore.GetObjectAsync(ns, metadataHash, cancellationToken: token);
-					byte[] b = await contents.Stream.ReadAllBytesAsync(cancellationToken: token);
-					CbObject o = new CbObject(b);
-					lock (responseObject)
+					try
+					{
+						await using BlobContents contents = await _blobStore.GetObjectAsync(ns, metadataHash, cancellationToken: token);
+						byte[] b = await contents.Stream.ReadAllBytesAsync(cancellationToken: token);
+						CbObject o = new CbObject(b);
+						lock (responseObject)
+						{
+							responseObject.WriteObject(o);
+						}
+					}
+					catch (BlobNotFoundException)
 					{
-						responseObject.WriteObject(o);
+						// if the metadata doesn't exist we can not display it
 					}
-				}
-				catch (BlobNotFoundException)
+				});
+			}
+			catch (AggregateException e)
+			{
+				if (e.InnerExceptions.Any(exception => exception.GetType() == typeof(RemoteServiceException)))
 				{
-					// if the metadata doesn't exist we can not display it
+					return Problem(e.Message, null, (int)HttpStatusCode.ServiceUnavailable);
 				}
-			});
+
+				throw;
+			}
+			catch (RemoteServiceException e)
+			{
+				return Problem(e.Message, null, (int)HttpStatusCode.ServiceUnavailable);
+			}
+
 			responseObject.EndArray();
 			responseObject.EndObject();
 
