Cloud DDC - Ignore case on the bucket regular expression used for search

diff --git a/Engine/Source/Programs/UnrealCloudDDC/Jupiter/Controllers/BuildsController.cs b/Engine/Source/Programs/UnrealCloudDDC/Jupiter/Controllers/BuildsController.cs
--- a/Engine/Source/Programs/UnrealCloudDDC/Jupiter/Controllers/BuildsController.cs
+++ b/Engine/Source/Programs/UnrealCloudDDC/Jupiter/Controllers/BuildsController.cs
@@ -476,7 +476,7 @@ namespace Jupiter.Controllers
 			try
 			{
 #pragma warning disable CA3012 // regex has a timeout so is okay to read from the body
-				regex = new Regex(searchBucketRequest.BucketRegex, RegexOptions.None, TimeSpan.FromSeconds(5));
+				regex = new Regex(searchBucketRequest.BucketRegex, RegexOptions.IgnoreCase, TimeSpan.FromSeconds(5));
 #pragma warning restore CA3012
 			}
 			catch (RegexParseException e)
