Cloud DDC - Added option to control when metadata blocks are GCed
Also added more debugging logging when orphan blobs are deleted if they never had any reference to them, this is a indication that something could be wrong (consistency issue with refs or that blobs are uploaded without ever pointing to them).

diff --git a/Engine/Source/Programs/UnrealCloudDDC/Jupiter/Implementation/GC/OrphanBlobCleanupRefs.cs b/Engine/Source/Programs/UnrealCloudDDC/Jupiter/Implementation/GC/OrphanBlobCleanupRefs.cs
--- a/Engine/Source/Programs/UnrealCloudDDC/Jupiter/Implementation/GC/OrphanBlobCleanupRefs.cs
+++ b/Engine/Source/Programs/UnrealCloudDDC/Jupiter/Implementation/GC/OrphanBlobCleanupRefs.cs
@@ -131,6 +131,7 @@ namespace Jupiter.Implementation
 				.SetAttribute("wasBlock", false.ToString());
 
 			bool found = false;
+			bool noReferences = true;
 
 			// check all namespaces that share the same storage pool for presence of the blob
 			foreach (NamespaceId blobNamespace in namespacesThatSharePool)
@@ -161,6 +162,7 @@ namespace Jupiter.Implementation
 						BucketId bucket = refBlobReference.Bucket;
 						RefId key = refBlobReference.Key;
 
+						noReferences = false;
 						try
 						{
 							bool refFound = await _refService.ExistsAsync(blobNamespace, bucket, key, cancellationToken);
@@ -170,6 +172,11 @@ namespace Jupiter.Implementation
 								found = true;
 								break;
 							}
+							else
+							{
+								// this is not a valid reference so we should delete
+								oldReferences.Add(refBlobReference);
+							}
 						}
 						catch (RefNotFoundException)
 						{
@@ -179,6 +186,7 @@ namespace Jupiter.Implementation
 					}
 					else if (baseBlobReference is BlobToBlobReference blobReference)
 					{
+						noReferences = false;
 						removeBlobScope.SetAttribute("hadBlobReference", bool.TrueString);
 
 						BlobId referringBlob = blobReference.Blob;
@@ -230,16 +238,25 @@ namespace Jupiter.Implementation
 				_logger.LogWarning("Recently added blob {Blob} from {StoragePool} was deleted. Indicates blob upload without anything that referenced it which is safe but can be a sign of an issue", blob, storagePoolName);
 			}
 
-			// delete block information if this blob is a block
-			foreach (NamespaceId namespaceId in namespacesThatSharePool)
+			if (noReferences)
 			{
-				BlobId? metadata = await _blockStore.GetBlockMetadataAsync(namespaceId, blob);
-				if (metadata != null)
+				// this blob had no references pointing to it which is not what we would expect, we would expect it to have some reference but that this reference in turn didn't exist
+				_logger.LogWarning("{Blob} from {StoragePool} was deleted and had no references to it. This is safe but a indication that data is added and never referenced", blob, storagePoolName);
+			}
+
+			if (_gcSettings.CurrentValue.AllowBlockMetadataCleanup)
+			{
+				// delete block information if this blob is a block
+				foreach (NamespaceId namespaceId in namespacesThatSharePool)
 				{
-					removeBlobScope.SetAttribute("wasBlock", true.ToString());
+					BlobId? metadata = await _blockStore.GetBlockMetadataAsync(namespaceId, blob);
+					if (metadata != null)
+					{
+						removeBlobScope.SetAttribute("wasBlock", true.ToString());
 
-					// this blob is a block, GC the block unique tracking
-					await _blockStore.DeleteBlockAsync(namespaceId, metadata);
+						// this blob is a block, GC the block unique tracking
+						await _blockStore.DeleteBlockAsync(namespaceId, metadata);
+					}
 				}
 			}
 
diff --git a/Engine/Source/Programs/UnrealCloudDDC/Jupiter/Settings.cs b/Engine/Source/Programs/UnrealCloudDDC/Jupiter/Settings.cs
--- a/Engine/Source/Programs/UnrealCloudDDC/Jupiter/Settings.cs
+++ b/Engine/Source/Programs/UnrealCloudDDC/Jupiter/Settings.cs
@@ -268,6 +268,7 @@ namespace Jupiter
 		public int OrphanRefMaxParallelOperations { get; set; } = 8;
 		public bool WriteDeleteToReplicationLog { get; set; } = false;
 		public NamespacePolicy.StoragePoolGCMethod DefaultGCPolicy { get; set; } = NamespacePolicy.StoragePoolGCMethod.LastAccess;
+		public bool AllowBlockMetadataCleanup { get; set; } = true;
 	}
 
 	public class UpstreamRelaySettings
