Cloud DDC - Option to control were incoming blob references are removed

diff --git a/Engine/Source/Programs/UnrealCloudDDC/Helm/UnrealCloudDDC/templates/configmap-worker.yaml b/Engine/Source/Programs/UnrealCloudDDC/Helm/UnrealCloudDDC/templates/configmap-worker.yaml
--- a/Engine/Source/Programs/UnrealCloudDDC/Helm/UnrealCloudDDC/templates/configmap-worker.yaml
+++ b/Engine/Source/Programs/UnrealCloudDDC/Helm/UnrealCloudDDC/templates/configmap-worker.yaml
@@ -34,6 +34,10 @@ UnrealCloudDDC:
   # And you would typically not want to extend the lifetime of refs based on them being accessed by scheduled tasks
   EnableLastAccessTracking: {{ default "false" $config.EnableLastAccessTracking }}
 
+{{- if hasKey $config "RemoveReferencesOnBlobDeletion" }}
+  RemoveReferencesOnBlobDeletion: {{ $config.RemoveReferencesOnBlobDeletion }}
+{{- end }}
+
 {{ if $config.StorageImplementations }}
   StorageImplementations:
   {{- range $config.StorageImplementations -}}
diff --git a/Engine/Source/Programs/UnrealCloudDDC/Helm/UnrealCloudDDC/templates/configmap.yaml b/Engine/Source/Programs/UnrealCloudDDC/Helm/UnrealCloudDDC/templates/configmap.yaml
--- a/Engine/Source/Programs/UnrealCloudDDC/Helm/UnrealCloudDDC/templates/configmap.yaml
+++ b/Engine/Source/Programs/UnrealCloudDDC/Helm/UnrealCloudDDC/templates/configmap.yaml
@@ -56,6 +56,10 @@ UnrealCloudDDC:
   UseLastAccessBackgroundService: {{ $config.UseLastAccessBackgroundService }}
 {{- end }}
 
+{{- if hasKey $config "RemoveReferencesOnBlobDeletion" }}
+  RemoveReferencesOnBlobDeletion: {{ $config.RemoveReferencesOnBlobDeletion }}
+{{- end }}
+
 {{ if $config.StorageImplementations }}
   StorageImplementations:
   {{- range $config.StorageImplementations -}}
diff --git a/Engine/Source/Programs/UnrealCloudDDC/Jupiter/Implementation/Blob/BlobService.cs b/Engine/Source/Programs/UnrealCloudDDC/Jupiter/Implementation/Blob/BlobService.cs
--- a/Engine/Source/Programs/UnrealCloudDDC/Jupiter/Implementation/Blob/BlobService.cs
+++ b/Engine/Source/Programs/UnrealCloudDDC/Jupiter/Implementation/Blob/BlobService.cs
@@ -1075,7 +1075,11 @@ public class BlobService : IBlobService
 		// if the blob store delete fails on the other hand we will still run a delete again during GC (as the blob is still orphaned at that point)
 		// this assumes that blob gc is based on scanning the root blob store
 		await _blobIndex.RemoveBlobFromAllRegionsAsync(ns, blob, cancellationToken: cancellationToken);
-		await _blobIndex.RemoveReferencesAsync(ns, blob, null, cancellationToken);
+
+		if (_settings.CurrentValue.RemoveReferencesOnBlobDeletion)
+		{
+			await _blobIndex.RemoveReferencesAsync(ns, blob, null, cancellationToken);
+		}
 
 		foreach (IBlobStore store in _blobStores)
 		{
@@ -1122,7 +1126,10 @@ public class BlobService : IBlobService
 		await Parallel.ForEachAsync(namespaces, cancellationToken, async (ns, token) =>
 		{
 			await _blobIndex.RemoveBlobFromAllRegionsAsync(ns, blob, cancellationToken: token);
-			await _blobIndex.RemoveReferencesAsync(ns, blob, null, token);
+			if (_settings.CurrentValue.RemoveReferencesOnBlobDeletion)
+			{
+				await _blobIndex.RemoveReferencesAsync(ns, blob, null, token);
+			}
 		});
 
 		// let each blob store figure out how it could effectively clear a storage pool
diff --git a/Engine/Source/Programs/UnrealCloudDDC/Jupiter/Implementation/GC/OrphanBlobCleanupRefs.cs b/Engine/Source/Programs/UnrealCloudDDC/Jupiter/Implementation/GC/OrphanBlobCleanupRefs.cs
--- a/Engine/Source/Programs/UnrealCloudDDC/Jupiter/Implementation/GC/OrphanBlobCleanupRefs.cs
+++ b/Engine/Source/Programs/UnrealCloudDDC/Jupiter/Implementation/GC/OrphanBlobCleanupRefs.cs
@@ -87,9 +87,8 @@ namespace Jupiter.Implementation
 
 				_logger.LogInformation("Running Orphan GC For StoragePool: {StoragePool}", policy.StoragePool);
 				namespacesThatHaveBeenChecked.AddRange(namespacesThatSharePool);
-				// only consider blobs that have been around for 60 minutes
-				// this due to cases were blobs are uploaded first
-				DateTime cutoff = DateTime.Now.AddMinutes(-60);
+				// only consider blobs that have been around for long enough
+				DateTime cutoff = DateTime.Now.AddMinutes(-1 * _gcSettings.CurrentValue.OrphanBlobCutoffDurationMinutes);
 				await Parallel.ForEachAsync(_blobService.ListObjectsAsync(@namespace, cancellationToken),
 					new ParallelOptions { MaxDegreeOfParallelism = _gcSettings.CurrentValue.OrphanGCMaxParallelOperations, CancellationToken = cancellationToken },
 					async (tuple, ctx) =>
diff --git a/Engine/Source/Programs/UnrealCloudDDC/Jupiter/Settings.cs b/Engine/Source/Programs/UnrealCloudDDC/Jupiter/Settings.cs
--- a/Engine/Source/Programs/UnrealCloudDDC/Jupiter/Settings.cs
+++ b/Engine/Source/Programs/UnrealCloudDDC/Jupiter/Settings.cs
@@ -141,6 +141,11 @@ namespace Jupiter
 		/// Set to enable using the new last access service which is registered as a background service
 		/// </summary>
 		public bool UseLastAccessBackgroundService { get; set; } = false;
+
+		/// <summary>
+		/// Remove the incoming blob references when the blob is deleted
+		/// </summary>
+		public bool RemoveReferencesOnBlobDeletion { get; set; } = true;
 	}
 
 	public class MongoSettings
@@ -269,6 +274,11 @@ namespace Jupiter
 		public bool WriteDeleteToReplicationLog { get; set; } = false;
 		public NamespacePolicy.StoragePoolGCMethod DefaultGCPolicy { get; set; } = NamespacePolicy.StoragePoolGCMethod.LastAccess;
 		public bool AllowBlockMetadataCleanup { get; set; } = true;
+
+		/// <summary>
+		/// How long a blob needs to be around before its considered for potential GC
+		/// </summary>
+		public int OrphanBlobCutoffDurationMinutes { get; set; } = 60;
 	}
 
 	public class UpstreamRelaySettings
