Cloud DDC - Fix incorrect type for dependency injection in ReplicationController

diff --git a/Engine/Source/Programs/UnrealCloudDDC/Jupiter/Controllers/ReplicationController.cs b/Engine/Source/Programs/UnrealCloudDDC/Jupiter/Controllers/ReplicationController.cs
--- a/Engine/Source/Programs/UnrealCloudDDC/Jupiter/Controllers/ReplicationController.cs
+++ b/Engine/Source/Programs/UnrealCloudDDC/Jupiter/Controllers/ReplicationController.cs
@@ -20,9 +20,9 @@ namespace Jupiter.Controllers
 	public class ReplicationController : ControllerBase
 	{
 		private readonly ReplicationService _replicationService;
-		private readonly RequestHelper _requestHelper;
+		private readonly IRequestHelper _requestHelper;
 
-		public ReplicationController(ReplicationService replicationService, RequestHelper requestHelper)
+		public ReplicationController(ReplicationService replicationService, IRequestHelper requestHelper)
 		{
 			_replicationService = replicationService;
 			_requestHelper = requestHelper;
