Cloud DDC - Improved error message and added attribute to better track when we hit a ref that is missing some of its blobs.

diff --git a/Engine/Source/Programs/UnrealCloudDDC/Jupiter/Controllers/ReferencesController.cs b/Engine/Source/Programs/UnrealCloudDDC/Jupiter/Controllers/ReferencesController.cs
--- a/Engine/Source/Programs/UnrealCloudDDC/Jupiter/Controllers/ReferencesController.cs
+++ b/Engine/Source/Programs/UnrealCloudDDC/Jupiter/Controllers/ReferencesController.cs
@@ -372,11 +372,14 @@ namespace Jupiter.Controllers
 								}
 								catch (PartialReferenceResolveException)
 								{
+									Tracer.CurrentSpan.SetAttribute("error-type", "MissingContentIds");
 									return NotFound(new ProblemDetails { Title = $"Object {bucket} {key} in namespace {ns} was missing some content ids" });
 								}
-								catch (ReferenceIsMissingBlobsException)
+								catch (ReferenceIsMissingBlobsException e)
 								{
-									return NotFound(new ProblemDetails { Title = $"Object {bucket} {key} in namespace {ns} was missing some blobs" });
+									BlobId? blobId = e.MissingBlobs.FirstOrDefault();
+									Tracer.CurrentSpan.SetAttribute("error-type", "MissingBlobs");
+									return NotFound(new ProblemDetails { Title = $"Object {bucket} {key} in namespace {ns} was missing {e.MissingBlobs.Count} blobs. Sample: {blobId}" });
 								}
 
 								if (referencedBlobs.Count == 1)
