Cloud DDC - Check that all parts are finalized before finalizing the build (making the build visible to be downloaded)

diff --git a/Engine/Source/Programs/UnrealCloudDDC/Jupiter/Controllers/BuildsController.cs b/Engine/Source/Programs/UnrealCloudDDC/Jupiter/Controllers/BuildsController.cs
--- a/Engine/Source/Programs/UnrealCloudDDC/Jupiter/Controllers/BuildsController.cs
+++ b/Engine/Source/Programs/UnrealCloudDDC/Jupiter/Controllers/BuildsController.cs
@@ -1,48 +1,48 @@
 // Copyright Epic Games, Inc. All Rights Reserved.
 
+using System;
+using System.Collections.Concurrent;
+using System.Collections.Generic;
+using System.ComponentModel.DataAnnotations;
+using System.Diagnostics;
+using System.Diagnostics.Metrics;
+using System.Globalization;
+using System.IO;
+using System.Linq;
+using System.Net;
+using System.Net.Http;
+using System.Net.Mime;
+using System.Text;
+using System.Text.Json;
+using System.Text.Json.Serialization;
+using System.Text.RegularExpressions;
+using System.Threading;
+using System.Threading.Tasks;
 using EpicGames.AspNet;
 using EpicGames.Core;
 using EpicGames.Horde.Storage;
 using EpicGames.Serialization;
-using Jupiter.Common.Implementation;
+using EpicGames.Serialization.Converters;
 using Jupiter.Common;
+using Jupiter.Common.Implementation;
 using Jupiter.Implementation;
+using Jupiter.Implementation.Blob;
+using Jupiter.Implementation.Builds;
 using Jupiter.Utils;
 using Microsoft.AspNetCore.Authorization;
-using Microsoft.AspNetCore.Http.Extensions;
 using Microsoft.AspNetCore.Http;
+using Microsoft.AspNetCore.Http.Extensions;
 using Microsoft.AspNetCore.Mvc;
 using Microsoft.Extensions.DependencyInjection;
 using Microsoft.Extensions.Logging;
+using Microsoft.Extensions.Options;
 using Microsoft.Extensions.Primitives;
 using OpenTelemetry.Trace;
 using Serilog;
-using System.Collections.Generic;
-using System.ComponentModel.DataAnnotations;
-using System.Globalization;
-using System.IO;
-using System.Linq;
-using System.Net.Mime;
-using System.Net;
-using System.Text;
-using System.Threading.Tasks;
-using System;
-using System.Collections.Concurrent;
-using System.Diagnostics;
-using System.Diagnostics.Metrics;
-using System.Net.Http;
-using System.Text.Json;
-using System.Text.Json.Serialization;
-using System.Text.RegularExpressions;
-using System.Threading;
-using Jupiter.Implementation.Blob;
-using Jupiter.Implementation.Builds;
-using Microsoft.Extensions.Options;
-using ContentHash = Jupiter.Implementation.ContentHash;
 using CbField = EpicGames.Serialization.CbField;
+using ContentHash = Jupiter.Implementation.ContentHash;
 using ILogger = Serilog.ILogger;
 using NotImplementedException = System.NotImplementedException;
-using EpicGames.Serialization.Converters;
 
 namespace Jupiter.Controllers
 {
@@ -195,6 +195,23 @@
 
 			try
 			{
+				await foreach ((string _, CbObjectId partId) in _buildStore.GetBuildPartsAsync(ns, bucket, buildId))
+				{
+					RefId refId = RefId.FromName($"{buildId}/{partId}");
+					try
+					{
+						(RefRecord refRecord, BlobContents? _)  = await _refService.GetAsync(ns, bucket, refId);
+						if (!refRecord.IsFinalized)
+						{
+							return BadRequest($"Build part {partId} ({refId}) for build {buildId} in bucket {bucket} was not finalized. Build can not be finalized.");
+						}
+					}
+					catch (RefNotFoundException)
+					{
+						return BadRequest($"Build part {partId} ({refId}) for build {buildId} in bucket {bucket} did not exist. Build can not be finalized.");
+					}
+				}
+				
 				await _buildStore.FinalizeBuildAsync(ns, bucket, buildId, ttl);
 			}
 			catch (RemoteServiceException e)
@@ -837,7 +854,7 @@
 				if (!objectRecord.IsFinalized)
 				{
 					// we do not consider un-finalized objects as valid
-					return BadRequest(new ProblemDetails { Title = $"Build {objectRecord.Bucket} {objectRecord.Name} is not finalized." });
+					return BadRequest(new ProblemDetails { Title = $"Build {buildId} in bucket {objectRecord.Bucket} with part {partId} stored as ref {objectRecord.Name} is not finalized." });
 				}
 
 				Response.Headers[CommonHeaders.HashHeaderName] = objectRecord.BlobIdentifier.ToString();
