Cloud Storage - Return 503 if scylla in unavailable when requesting namespaces

diff --git a/Engine/Source/Programs/UnrealCloudDDC/Jupiter/Controllers/BuildsController.cs b/Engine/Source/Programs/UnrealCloudDDC/Jupiter/Controllers/BuildsController.cs
--- a/Engine/Source/Programs/UnrealCloudDDC/Jupiter/Controllers/BuildsController.cs
+++ b/Engine/Source/Programs/UnrealCloudDDC/Jupiter/Controllers/BuildsController.cs
@@ -1909,7 +1909,15 @@ namespace Jupiter.Controllers
 		[ProducesDefaultResponseType]
 		public async Task<IActionResult> GetNamespacesAsync()
 		{
-			NamespaceId[] namespaces = await _refService.GetNamespacesAsync(HttpContext.RequestAborted).ToArrayAsync(HttpContext.RequestAborted);
+			NamespaceId[] namespaces;
+			try
+			{
+				namespaces = await _refService.GetNamespacesAsync(HttpContext.RequestAborted).ToArrayAsync(HttpContext.RequestAborted);
+			}
+			catch (RemoteServiceException e)
+			{
+				return Problem(e.Message, null, (int)HttpStatusCode.ServiceUnavailable);
+			}
 
 			// filter namespaces down to only the namespaces the user has access to
 			List<NamespaceId> namespacesWithAccess = new();
