Cloud DDC - Only persist new replicator state when old state matches the expected previous state. This allows admin endpoints to mutate replicator state without worrying about being overwritten.

#Synced-CL 48106205

diff --git a/Engine/Source/Programs/UnrealCloudDDC/Jupiter/Implementation/Replication/BlobsReplicator.cs b/Engine/Source/Programs/UnrealCloudDDC/Jupiter/Implementation/Replication/BlobsReplicator.cs
--- a/Engine/Source/Programs/UnrealCloudDDC/Jupiter/Implementation/Replication/BlobsReplicator.cs
+++ b/Engine/Source/Programs/UnrealCloudDDC/Jupiter/Implementation/Replication/BlobsReplicator.cs
@@ -143,9 +143,31 @@ namespace Jupiter.Implementation
 			_disposed = true;
 		}
 
-		private async Task SaveStateAsync(BlobsState newState)
+		/// <summary>
+		/// Updates the state of the current replicator using the provided new state.
+		/// 
+		/// If an expected old state is provided, then ensures that the existing state matches prior to saving the new state. This allows the admin
+		/// endpoints to rewind or fast-forward the replicator's state without worrying about being overwritten.
+		///
+		/// If an existing state does not exist at all, then the new state will be inserted regardless of the provided expected old state.
+		/// </summary>
+		/// <param name="newState">New state to persist</param>
+		/// <param name="expectedOldState">Optional old state that is expected to exist before persisting the new state</param>
+		private async Task<bool> SaveStateAsync(BlobsState newState, BlobsState? expectedOldState = null)
 		{
+			if (expectedOldState != null)
+			{
+				ReplicatorState? existingState = await _replicationLog.GetReplicatorStateAsync(Info.NamespaceToReplicate, _name);
+
+				if (existingState != null && existingState.LastBucket != expectedOldState.LastBucket)
+				{
+					_logger.LogWarning("Existing replicator state for {Name} does not match expected value ({Expected} != {Actual}). Skipping persistence of new state.", _name, expectedOldState.LastBucket, existingState.LastBucket);
+					return false;
+				}
+			}
+			
 			await _replicationLog.UpdateReplicatorStateAsync(Info.NamespaceToReplicate, _name, new ReplicatorState { LastBucket = newState.LastBucket });
+			return true;
 		}
 
 		public async Task<bool> TriggerNewReplicationsAsync()
@@ -323,8 +345,14 @@ namespace Jupiter.Implementation
 				else
 				{
 					// we have replicated everything up to a point and can persist this in the state
+					BlobsState previousState = new BlobsState() { LastBucket = _replicationState.LastBucket };
 					_replicationState.LastBucket = refBucket;
-					await SaveStateAsync(_replicationState);
+
+					if (!await SaveStateAsync(_replicationState, expectedOldState: previousState))
+					{
+						_logger.LogWarning("Replication state for {Name} has changed during processing. Aborting.", _name);
+						break;
+					}
 				}
 
 				Info.LastRun = DateTime.Now;
