Cloud DDC - Add missing replicator version configuration to Helm chart

#Synced-CL 48106205

diff --git a/Engine/Source/Programs/UnrealCloudDDC/Helm/UnrealCloudDDC/templates/configmap.yaml b/Engine/Source/Programs/UnrealCloudDDC/Helm/UnrealCloudDDC/templates/configmap.yaml
--- a/Engine/Source/Programs/UnrealCloudDDC/Helm/UnrealCloudDDC/templates/configmap.yaml
+++ b/Engine/Source/Programs/UnrealCloudDDC/Helm/UnrealCloudDDC/templates/configmap.yaml
@@ -115,6 +115,7 @@ Replication:
     EnableMultipartReplications: {{ default false  .EnableMultipartReplications }}
     SkipSnapshot: {{ default true  .SkipSnapshot }}
     PageSize: {{ default 1000  .PageSize }}
+    Version: {{ default "Refs"  .Version }}
     {{- if .SourceSite }}
     SourceSite: {{ .SourceSite }}
     {{- else if .ConnectionString }}
