Cloud DDC - Change authz for ddc to use the new auth helper which includes support for per bucket authorization. We generally do not recommend configuring per bucket auth for ddc but this allows the new auth config to apply for ddc allowing us to fully deprecate the legacy configuration method.

#Synced-CL 47492933

diff --git a/Engine/Source/Programs/UnrealCloudDDC/Jupiter/Controllers/AdminController.cs b/Engine/Source/Programs/UnrealCloudDDC/Jupiter/Controllers/AdminController.cs
--- a/Engine/Source/Programs/UnrealCloudDDC/Jupiter/Controllers/AdminController.cs
+++ b/Engine/Source/Programs/UnrealCloudDDC/Jupiter/Controllers/AdminController.cs
@@ -39,7 +39,7 @@ namespace Jupiter.Controllers
 		[HttpPost("refCleanup")]
 		public async Task<IActionResult> RefCleanupAsync()
 		{
-			ActionResult? result = await _requestHelper.HasAccessForGlobalOperationsAsync(User, new[] { JupiterAclAction.AdminAction });
+			ActionResult? result = await _requestHelper.HasAccessToScopeAsync(User, Request, AccessScope.GlobalScope, new[] { JupiterAclAction.AdminAction });
 			if (result != null)
 			{	
 				return result;
@@ -56,7 +56,7 @@ namespace Jupiter.Controllers
 		[HttpGet("settings")]
 		public async Task<IActionResult> SettingsAsync()
 		{
-			ActionResult? result = await _requestHelper.HasAccessForGlobalOperationsAsync(User, new[] { JupiterAclAction.AdminAction });
+			ActionResult? result = await _requestHelper.HasAccessToScopeAsync(User, Request, AccessScope.GlobalScope, new[] { JupiterAclAction.AdminAction });
 			if (result != null)
 			{
 				return result;
diff --git a/Engine/Source/Programs/UnrealCloudDDC/Jupiter/Controllers/AuthController.cs b/Engine/Source/Programs/UnrealCloudDDC/Jupiter/Controllers/AuthController.cs
--- a/Engine/Source/Programs/UnrealCloudDDC/Jupiter/Controllers/AuthController.cs
+++ b/Engine/Source/Programs/UnrealCloudDDC/Jupiter/Controllers/AuthController.cs
@@ -71,10 +71,10 @@ namespace Jupiter.Controllers
 		[HttpGet("{ns}")]
 		public async Task<IActionResult> VerifyAsync([FromRoute][Required] NamespaceId ns)
 		{
-			ActionResult? result = await _requestHelper.HasAccessToNamespaceAsync(User, Request, ns, new[] { JupiterAclAction.ReadObject });
-			if (result != null)
+			ActionResult? accessResult = await _requestHelper.HasAccessToScopeAsync(User, Request, new AccessScope(ns), new[] { JupiterAclAction.ReadObject });
+			if (accessResult != null)
 			{
-				return result;
+				return accessResult;
 			}
 
 			return Ok();
diff --git a/Engine/Source/Programs/UnrealCloudDDC/Jupiter/Controllers/BlobController.cs b/Engine/Source/Programs/UnrealCloudDDC/Jupiter/Controllers/BlobController.cs
--- a/Engine/Source/Programs/UnrealCloudDDC/Jupiter/Controllers/BlobController.cs
+++ b/Engine/Source/Programs/UnrealCloudDDC/Jupiter/Controllers/BlobController.cs
@@ -56,10 +56,10 @@ namespace Jupiter.Controllers
 			[FromQuery] bool supportsRedirect = false,
 			[FromQuery] bool allowOndemandReplication = true)
 		{
-			ActionResult? result = await _requestHelper.HasAccessToNamespaceAsync(User, Request, ns, new[] { JupiterAclAction.ReadObject });
-			if (result != null)
+			ActionResult? accessResult = await _requestHelper.HasAccessToScopeAsync(User, Request, new AccessScope(ns), new[] { JupiterAclAction.ReadObject });
+			if (accessResult != null)
 			{
-				return result;
+				return accessResult;
 			}
 
 			try
@@ -97,10 +97,10 @@ namespace Jupiter.Controllers
 			[Required] BlobId id,
 			[FromQuery] List<string>? storageLayers = null)
 		{
-			ActionResult? result = await _requestHelper.HasAccessToNamespaceAsync(User, Request, ns, new[] { JupiterAclAction.ReadObject });
-			if (result != null)
+			ActionResult? accessResult = await _requestHelper.HasAccessToScopeAsync(User, Request, new AccessScope(ns), new[] { JupiterAclAction.ReadObject });
+			if (accessResult != null)
 			{
-				return result;
+				return accessResult;
 			}
 
 			try
@@ -126,10 +126,10 @@ namespace Jupiter.Controllers
 			[Required] NamespaceId ns,
 			[Required][FromQuery] List<BlobId> id)
 		{
-			ActionResult? result = await _requestHelper.HasAccessToNamespaceAsync(User, Request, ns, new[] { JupiterAclAction.ReadObject });
-			if (result != null)
+			ActionResult? accessResult = await _requestHelper.HasAccessToScopeAsync(User, Request, new AccessScope(ns), new[] { JupiterAclAction.ReadObject });
+			if (accessResult != null)
 			{
-				return result;
+				return accessResult;
 			}
 
 			ConcurrentBag<BlobId> missingBlobs = new ConcurrentBag<BlobId>();
@@ -152,10 +152,10 @@ namespace Jupiter.Controllers
 			[Required] NamespaceId ns,
 			[FromBody] BlobId[] bodyIds)
 		{
-			ActionResult? result = await _requestHelper.HasAccessToNamespaceAsync(User, Request, ns, new[] { JupiterAclAction.ReadObject });
-			if (result != null)
+			ActionResult? accessResult = await _requestHelper.HasAccessToScopeAsync(User, Request, new AccessScope(ns), new[] { JupiterAclAction.ReadObject });
+			if (accessResult != null)
 			{
-				return result;
+				return accessResult;
 			}
 
 			ConcurrentBag<BlobId> missingBlobs = new ConcurrentBag<BlobId>();
@@ -196,10 +196,10 @@ namespace Jupiter.Controllers
 			[Required] NamespaceId ns,
 			[Required] BlobId id)
 		{
-			ActionResult? result = await _requestHelper.HasAccessToNamespaceAsync(User, Request, ns, new[] { JupiterAclAction.WriteObject });
-			if (result != null)
+			ActionResult? accessResult = await _requestHelper.HasAccessToScopeAsync(User, Request, new AccessScope(ns), new[] { JupiterAclAction.WriteObject });
+			if (accessResult != null)
 			{
-				return result;
+				return accessResult;
 			}
 
 			_diagnosticContext.Set("Content-Length", Request.ContentLength ?? -1);
@@ -238,10 +238,10 @@ namespace Jupiter.Controllers
 		public async Task<IActionResult> PostAsync(
 			[Required] NamespaceId ns)
 		{
-			ActionResult? result = await _requestHelper.HasAccessToNamespaceAsync(User, Request, ns, new[] { JupiterAclAction.WriteObject });
-			if (result != null)
+			ActionResult? accessResult = await _requestHelper.HasAccessToScopeAsync(User, Request, new AccessScope(ns), new[] { JupiterAclAction.WriteObject });
+			if (accessResult != null)
 			{
-				return result;
+				return accessResult;
 			}
 
 			_diagnosticContext.Set("Content-Length", Request.ContentLength ?? -1);
@@ -272,10 +272,10 @@ namespace Jupiter.Controllers
 			[Required] NamespaceId ns,
 			[Required] BlobId id)
 		{
-			ActionResult? result = await _requestHelper.HasAccessToNamespaceAsync(User, Request, ns, new[] { JupiterAclAction.DeleteObject });
-			if (result != null)
+			ActionResult? accessResult = await _requestHelper.HasAccessToScopeAsync(User, Request, new AccessScope(ns), new[] { JupiterAclAction.DeleteObject });
+			if (accessResult != null)
 			{
-				return result;
+				return accessResult;
 			}
 
 			try
@@ -294,10 +294,10 @@ namespace Jupiter.Controllers
 		public async Task<IActionResult> DeleteNamespaceAsync(
 			[Required] NamespaceId ns)
 		{
-			ActionResult? result = await _requestHelper.HasAccessToNamespaceAsync(User, Request, ns, new[] { JupiterAclAction.DeleteNamespace });
-			if (result != null)
+			ActionResult? accessResult = await _requestHelper.HasAccessToScopeAsync(User, Request, new AccessScope(ns), new[] { JupiterAclAction.DeleteNamespace });
+			if (accessResult != null)
 			{
-				return result;
+				return accessResult;
 			}
 
 			try
@@ -381,7 +381,7 @@ namespace Jupiter.Controllers
 					throw new Exception(nameof(op.Namespace));
 				}
 
-				ActionResult? result = await _requestHelper.HasAccessToNamespaceAsync(User, Request, op.Namespace!.Value, new[] { MapToAclAction(op.Op) });
+				ActionResult? result = await _requestHelper.HasAccessToScopeAsync(User, Request, new AccessScope(op.Namespace.Value), new[] { MapToAclAction(op.Op) });
 
 				if (result != null)
 				{
diff --git a/Engine/Source/Programs/UnrealCloudDDC/Jupiter/Controllers/CompressedBlobController.cs b/Engine/Source/Programs/UnrealCloudDDC/Jupiter/Controllers/CompressedBlobController.cs
--- a/Engine/Source/Programs/UnrealCloudDDC/Jupiter/Controllers/CompressedBlobController.cs
+++ b/Engine/Source/Programs/UnrealCloudDDC/Jupiter/Controllers/CompressedBlobController.cs
@@ -59,10 +59,10 @@ namespace Jupiter.Controllers
 			[Required] ContentId id,
 			[FromQuery] bool supportsRedirect = false)
 		{
-			ActionResult? result = await _requestHelper.HasAccessToNamespaceAsync(User, Request, ns, new[] { JupiterAclAction.ReadObject });
-			if (result != null)
+			ActionResult? accessResult = await _requestHelper.HasAccessToScopeAsync(User, Request, new AccessScope(ns), new[] { JupiterAclAction.ReadObject });
+			if (accessResult != null)
 			{
-				return result;
+				return accessResult;
 			}
 
 			Tracer.CurrentSpan.SetAttribute("namespace", ns.ToString());
@@ -115,10 +115,10 @@ namespace Jupiter.Controllers
 			[Required] NamespaceId ns,
 			[Required] ContentId id)
 		{
-			ActionResult? result = await _requestHelper.HasAccessToNamespaceAsync(User, Request, ns, new[] { JupiterAclAction.ReadObject });
-			if (result != null)
+			ActionResult? accessResult = await _requestHelper.HasAccessToScopeAsync(User, Request, new AccessScope(ns), new[] { JupiterAclAction.ReadObject });
+			if (accessResult != null)
 			{
-				return result;
+				return accessResult;
 			}
 
 			BlobId[]? chunks = await _contentIdStore.ResolveAsync(ns, id, mustBeContentId: false, cancellationToken: HttpContext.RequestAborted);
@@ -151,10 +151,10 @@ namespace Jupiter.Controllers
 			[Required] NamespaceId ns,
 			[Required][FromQuery] List<ContentId> id)
 		{
-			ActionResult? result = await _requestHelper.HasAccessToNamespaceAsync(User, Request, ns, new[] { JupiterAclAction.ReadObject });
-			if (result != null)
+			ActionResult? accessResult = await _requestHelper.HasAccessToScopeAsync(User, Request, new AccessScope(ns), new[] { JupiterAclAction.ReadObject });
+			if (accessResult != null)
 			{
-				return result;
+				return accessResult;
 			}
 
 			ConcurrentBag<ContentId> partialContentIds = new ConcurrentBag<ContentId>();
@@ -193,10 +193,10 @@ namespace Jupiter.Controllers
 			[Required] NamespaceId ns,
 			[FromBody] ContentId[] bodyIds)
 		{
-			ActionResult? result = await _requestHelper.HasAccessToNamespaceAsync(User, Request, ns, new[] { JupiterAclAction.ReadObject });
-			if (result != null)
+			ActionResult? accessResult = await _requestHelper.HasAccessToScopeAsync(User, Request, new AccessScope(ns), new[] { JupiterAclAction.ReadObject });
+			if (accessResult != null)
 			{
-				return result;
+				return accessResult;
 			}
 
 			ConcurrentBag<ContentId> partialContentIds = new ConcurrentBag<ContentId>();
@@ -235,10 +235,10 @@ namespace Jupiter.Controllers
 			[Required] NamespaceId ns,
 			[Required] ContentId id)
 		{
-			ActionResult? result = await _requestHelper.HasAccessToNamespaceAsync(User, Request, ns, new[] { JupiterAclAction.WriteObject });
-			if (result != null)
+			ActionResult? accessResult = await _requestHelper.HasAccessToScopeAsync(User, Request, new AccessScope(ns), new[] { JupiterAclAction.WriteObject });
+			if (accessResult != null)
 			{
-				return result;
+				return accessResult;
 			}
 
 			Tracer.CurrentSpan.SetAttribute("namespace", ns.ToString());
@@ -279,12 +279,11 @@ namespace Jupiter.Controllers
 		{
 			CancellationToken cancellationToken = HttpContext.RequestAborted;
 
-			ActionResult? result = await _requestHelper.HasAccessToNamespaceAsync(User, Request, ns, new[] { JupiterAclAction.WriteObject });
-			if (result != null)
+			ActionResult? accessResult = await _requestHelper.HasAccessToScopeAsync(User, Request, new AccessScope(ns), new[] { JupiterAclAction.WriteObject });
+			if (accessResult != null)
 			{
-				return result;
+				return accessResult;
 			}
-
 			_diagnosticContext.Set("Content-Length", Request.ContentLength ?? -1);
 
 			try
@@ -320,10 +319,10 @@ namespace Jupiter.Controllers
 			[Required] NamespaceId ns,
 			[Required] ContentId id)
 		{
-			ActionResult? result = await _requestHelper.HasAccessToNamespaceAsync(User, Request, ns, new[] { JupiterAclAction.DeleteObject });
-			if (result != null)
+			ActionResult? accessResult = await _requestHelper.HasAccessToScopeAsync(User, Request, new AccessScope(ns), new[] { JupiterAclAction.DeleteObject });
+			if (accessResult != null)
 			{
-				return result;
+				return accessResult;
 			}
 
 			BlobId[]? chunks = await _contentIdStore.ResolveAsync(ns, id, mustBeContentId: true);
diff --git a/Engine/Source/Programs/UnrealCloudDDC/Jupiter/Controllers/ContentIdController.cs b/Engine/Source/Programs/UnrealCloudDDC/Jupiter/Controllers/ContentIdController.cs
--- a/Engine/Source/Programs/UnrealCloudDDC/Jupiter/Controllers/ContentIdController.cs
+++ b/Engine/Source/Programs/UnrealCloudDDC/Jupiter/Controllers/ContentIdController.cs
@@ -39,10 +39,10 @@ namespace Jupiter.Controllers
 		[ProducesResponseType(type: typeof(ProblemDetails), 400)]
 		public async Task<IActionResult> ResolveAsync(NamespaceId ns, ContentId contentId)
 		{
-			ActionResult? result = await _requestHelper.HasAccessToNamespaceAsync(User, Request, ns, new[] { JupiterAclAction.ReadObject });
-			if (result != null)
+			ActionResult? accessResult = await _requestHelper.HasAccessToScopeAsync(User, Request, new AccessScope(ns), new[] { JupiterAclAction.ReadObject });
+			if (accessResult != null)
 			{
-				return result;
+				return accessResult;
 			}
 
 			try
@@ -74,10 +74,10 @@ namespace Jupiter.Controllers
 		[ProducesResponseType(type: typeof(ProblemDetails), 400)]
 		public async Task<IActionResult> UpdateContentIdMappingAsync(NamespaceId ns, ContentId contentId, BlobId blobIdentifier, int contentWeight)
 		{
-			ActionResult? result = await _requestHelper.HasAccessToNamespaceAsync(User, Request, ns, new[] { JupiterAclAction.WriteObject });
-			if (result != null)
+			ActionResult? accessResult = await _requestHelper.HasAccessToScopeAsync(User, Request, new AccessScope(ns), new[] { JupiterAclAction.WriteObject });
+			if (accessResult != null)
 			{
-				return result;
+				return accessResult;
 			}
 
 			try
diff --git a/Engine/Source/Programs/UnrealCloudDDC/Jupiter/Controllers/IRequestHelper.cs b/Engine/Source/Programs/UnrealCloudDDC/Jupiter/Controllers/IRequestHelper.cs
--- a/Engine/Source/Programs/UnrealCloudDDC/Jupiter/Controllers/IRequestHelper.cs
+++ b/Engine/Source/Programs/UnrealCloudDDC/Jupiter/Controllers/IRequestHelper.cs
@@ -3,7 +3,6 @@
 using System;
 using System.Security.Claims;
 using System.Threading.Tasks;
-using EpicGames.Horde.Storage;
 using Microsoft.AspNetCore.Http;
 using Microsoft.AspNetCore.Mvc;
 
@@ -12,8 +11,6 @@ namespace Jupiter.Controllers
 	public interface IRequestHelper
 	{
 		public Task<ActionResult?> HasAccessToScopeAsync(ClaimsPrincipal user, HttpRequest request, AccessScope scope, JupiterAclAction[] aclActions);
-		public Task<ActionResult?> HasAccessToNamespaceAsync(ClaimsPrincipal user, HttpRequest request, NamespaceId ns, JupiterAclAction[] aclActions);
-		public Task<ActionResult?> HasAccessForGlobalOperationsAsync(ClaimsPrincipal user, JupiterAclAction[] aclActions);
 
 		public bool IsPublicPort(HttpContext context);
 	}
diff --git a/Engine/Source/Programs/UnrealCloudDDC/Jupiter/Controllers/ObjectController.cs b/Engine/Source/Programs/UnrealCloudDDC/Jupiter/Controllers/ObjectController.cs
--- a/Engine/Source/Programs/UnrealCloudDDC/Jupiter/Controllers/ObjectController.cs
+++ b/Engine/Source/Programs/UnrealCloudDDC/Jupiter/Controllers/ObjectController.cs
@@ -54,10 +54,10 @@ namespace Jupiter.Controllers
 			[Required] NamespaceId ns,
 			[Required] BlobId id)
 		{
-			ActionResult? result = await _requestHelper.HasAccessToNamespaceAsync(User, Request, ns, new[] { JupiterAclAction.ReadObject });
-			if (result != null)
+			ActionResult? accessResult = await _requestHelper.HasAccessToScopeAsync(User, Request, new AccessScope(ns), new[] { JupiterAclAction.ReadObject });
+			if (accessResult != null)
 			{
-				return result;
+				return accessResult;
 			}
 
 			try
@@ -82,12 +82,12 @@ namespace Jupiter.Controllers
 			[Required] NamespaceId ns,
 			[Required] BlobId id)
 		{
-			ActionResult? result = await _requestHelper.HasAccessToNamespaceAsync(User, Request, ns, new[] { JupiterAclAction.ReadObject });
-			if (result != null)
+			ActionResult? accessResult = await _requestHelper.HasAccessToScopeAsync(User, Request, new AccessScope(ns), new[] { JupiterAclAction.ReadObject });
+			if (accessResult != null)
 			{
-				return result;
+				return accessResult;
 			}
-
+			
 			try
 			{
 				bool exists = await _storage.ExistsAsync(ns, id);
@@ -111,10 +111,10 @@ namespace Jupiter.Controllers
 			[Required] NamespaceId ns,
 			[Required][FromQuery] List<BlobId> id)
 		{
-			ActionResult? result = await _requestHelper.HasAccessToNamespaceAsync(User, Request, ns, new[] { JupiterAclAction.ReadObject });
-			if (result != null)
+			ActionResult? accessResult = await _requestHelper.HasAccessToScopeAsync(User, Request, new AccessScope(ns), new[] { JupiterAclAction.ReadObject });
+			if (accessResult != null)
 			{
-				return result;
+				return accessResult;
 			}
 
 			ConcurrentBag<BlobId> missingBlobs = new ConcurrentBag<BlobId>();
@@ -137,10 +137,10 @@ namespace Jupiter.Controllers
 			[Required] NamespaceId ns,
 			[FromBody] BlobId[] bodyIds)
 		{
-			ActionResult? result = await _requestHelper.HasAccessToNamespaceAsync(User, Request, ns, new[] { JupiterAclAction.ReadObject });
-			if (result != null)
+			ActionResult? accessResult = await _requestHelper.HasAccessToScopeAsync(User, Request, new AccessScope(ns), new[] { JupiterAclAction.ReadObject });
+			if (accessResult != null)
 			{
-				return result;
+				return accessResult;
 			}
 
 			ConcurrentBag<BlobId> missingBlobs = new ConcurrentBag<BlobId>();
@@ -163,10 +163,10 @@ namespace Jupiter.Controllers
 			[Required] NamespaceId ns,
 			[Required] BlobId id)
 		{
-			ActionResult? result = await _requestHelper.HasAccessToNamespaceAsync(User, Request, ns, new[] { JupiterAclAction.WriteObject });
-			if (result != null)
+			ActionResult? accessResult = await _requestHelper.HasAccessToScopeAsync(User, Request, new AccessScope(ns), new[] { JupiterAclAction.WriteObject });
+			if (accessResult != null)
 			{
-				return result;
+				return accessResult;
 			}
 
 			_diagnosticContext.Set("Content-Length", Request.ContentLength ?? -1);
@@ -192,10 +192,10 @@ namespace Jupiter.Controllers
 			[Required] NamespaceId ns,
 			[Required] BlobId id)
 		{
-			ActionResult? result = await _requestHelper.HasAccessToNamespaceAsync(User, Request, ns, new[] { JupiterAclAction.ReadObject });
-			if (result != null)
+			ActionResult? accessResult = await _requestHelper.HasAccessToScopeAsync(User, Request, new AccessScope(ns), new[] { JupiterAclAction.ReadObject });
+			if (accessResult != null)
 			{
-				return result;
+				return accessResult;
 			}
 
 			BlobContents blob;
@@ -252,10 +252,10 @@ namespace Jupiter.Controllers
 			[Required] NamespaceId ns,
 			[Required] BlobId id)
 		{
-			ActionResult? result = await _requestHelper.HasAccessToNamespaceAsync(User, Request, ns, new[] { JupiterAclAction.DeleteObject });
-			if (result != null)
+			ActionResult? accessResult = await _requestHelper.HasAccessToScopeAsync(User, Request, new AccessScope(ns), new[] { JupiterAclAction.DeleteObject });
+			if (accessResult != null)
 			{
-				return result;
+				return accessResult;
 			}
 
 			try
@@ -277,10 +277,10 @@ namespace Jupiter.Controllers
 		public async Task<IActionResult> DeleteNamespaceAsync(
 			[Required] NamespaceId ns)
 		{
-			ActionResult? result = await _requestHelper.HasAccessToNamespaceAsync(User, Request, ns, new[] { JupiterAclAction.DeleteNamespace });
-			if (result != null)
+			ActionResult? accessResult = await _requestHelper.HasAccessToScopeAsync(User, Request, new AccessScope(ns), new[] { JupiterAclAction.DeleteNamespace });
+			if (accessResult != null)
 			{
-				return result;
+				return accessResult;
 			}
 
 			try
diff --git a/Engine/Source/Programs/UnrealCloudDDC/Jupiter/Controllers/ReferencesController.cs b/Engine/Source/Programs/UnrealCloudDDC/Jupiter/Controllers/ReferencesController.cs
--- a/Engine/Source/Programs/UnrealCloudDDC/Jupiter/Controllers/ReferencesController.cs
+++ b/Engine/Source/Programs/UnrealCloudDDC/Jupiter/Controllers/ReferencesController.cs
@@ -94,7 +94,7 @@ namespace Jupiter.Controllers
 				List<NamespaceId> namespacesWithAccess = new();
 				foreach (NamespaceId ns in namespaces)
 				{
-					ActionResult? accessResult = await _requestHelper.HasAccessToNamespaceAsync(User, Request, ns, new[] { JupiterAclAction.ReadObject });
+					ActionResult? accessResult = await _requestHelper.HasAccessToScopeAsync(User, Request, new AccessScope(ns), new[] { JupiterAclAction.ReadObject });
 					if (accessResult == null)
 					{
 						namespacesWithAccess.Add(ns);
@@ -129,7 +129,7 @@ namespace Jupiter.Controllers
 			[FromRoute][Required] RefId key,
 			[FromRoute] string? format = null)
 		{
-			ActionResult? accessResult = await _requestHelper.HasAccessToNamespaceAsync(User, Request, ns, new[] { JupiterAclAction.ReadObject });
+			ActionResult? accessResult = await _requestHelper.HasAccessToScopeAsync(User, Request, new AccessScope(ns, bucket), new[] { JupiterAclAction.ReadObject });
 			if (accessResult != null)
 			{
 				return accessResult;
@@ -478,7 +478,7 @@ namespace Jupiter.Controllers
 			[FromRoute][Required] RefId key,
 			[FromQuery] string[] fields)
 		{
-			ActionResult? accessResult = await _requestHelper.HasAccessToNamespaceAsync(User, Request, ns, new[] { JupiterAclAction.ReadObject });
+			ActionResult? accessResult = await _requestHelper.HasAccessToScopeAsync(User, Request, new AccessScope(ns, bucket), new[] { JupiterAclAction.ReadObject });
 			if (accessResult != null)
 			{
 				return accessResult;
@@ -520,7 +520,7 @@ namespace Jupiter.Controllers
 			[FromRoute][Required] BucketId bucket,
 			[FromRoute][Required] RefId key)
 		{
-			ActionResult? accessResult = await _requestHelper.HasAccessToNamespaceAsync(User, Request, ns, new[] { JupiterAclAction.AdminAction });
+			ActionResult? accessResult = await _requestHelper.HasAccessToScopeAsync(User, Request, new AccessScope(ns, bucket), new[] { JupiterAclAction.AdminAction });
 			if (accessResult != null)
 			{
 				return accessResult;
@@ -579,7 +579,7 @@ namespace Jupiter.Controllers
 			[FromRoute][Required] BucketId bucket,
 			[FromRoute][Required] RefId key)
 		{
-			ActionResult? accessResult = await _requestHelper.HasAccessToNamespaceAsync(User, Request, ns, new[] { JupiterAclAction.ReadObject });
+			ActionResult? accessResult = await _requestHelper.HasAccessToScopeAsync(User, Request, new AccessScope(ns, bucket), new[] { JupiterAclAction.ReadObject });
 			if (accessResult != null)
 			{
 				return accessResult;
@@ -669,7 +669,7 @@ namespace Jupiter.Controllers
 			[FromRoute][Required] BucketId bucket,
 			[FromRoute][Required] RefId key)
 		{
-			ActionResult? accessResult = await _requestHelper.HasAccessToNamespaceAsync(User, Request, ns, new[] { JupiterAclAction.ReadObject });
+			ActionResult? accessResult = await _requestHelper.HasAccessToScopeAsync(User, Request, new AccessScope(ns, bucket), new[] { JupiterAclAction.ReadObject });
 			if (accessResult != null)
 			{
 				return accessResult;
@@ -737,12 +737,11 @@ namespace Jupiter.Controllers
 			[FromRoute][Required] NamespaceId ns,
 			[FromQuery][Required] List<string> names)
 		{
-			ActionResult? accessResult = await _requestHelper.HasAccessToNamespaceAsync(User, Request, ns, new[] { JupiterAclAction.ReadObject });
-			if (accessResult != null)
+			ActionResult? accessResultNamespace = await _requestHelper.HasAccessToScopeAsync(User, Request, new AccessScope(ns), new[] { JupiterAclAction.ReadObject });
+			if (accessResultNamespace != null)
 			{
-				return accessResult;
+				return accessResultNamespace;
 			}
-
 			ConcurrentBag<(BucketId, RefId)> missingObject = new();
 
 			List<(BucketId, RefId)> requestedNames = new List<(BucketId, RefId)>();
@@ -759,6 +758,16 @@ namespace Jupiter.Controllers
 				requestedNames.Add((bucket, key));
 			}
 
+			// verify we have access to each of the buckets we are requesting info on
+			foreach (BucketId bucket in requestedNames.Select(tuple => tuple.Item1).Distinct())
+			{
+				ActionResult? accessResult = await _requestHelper.HasAccessToScopeAsync(User, Request, new AccessScope(ns, bucket), new[] { JupiterAclAction.ReadObject });
+				if (accessResult != null)
+				{
+					return accessResult;
+				}
+			}
+
 			IEnumerable<Task> tasks = requestedNames.Select(async pair =>
 			{
 				(BucketId bucket, RefId key) = pair;
@@ -807,7 +816,7 @@ namespace Jupiter.Controllers
 			[FromRoute][Required] RefId key,
 			[FromHeader(Name = CommonHeaders.AllowOverwrite)] bool allowOverwrite = false)
 		{
-			ActionResult? accessResult = await _requestHelper.HasAccessToNamespaceAsync(User, Request, ns, new[] { JupiterAclAction.WriteObject });
+			ActionResult? accessResult = await _requestHelper.HasAccessToScopeAsync(User, Request, new AccessScope(ns, bucket), new[] { JupiterAclAction.WriteObject });
 			if (accessResult != null)
 			{
 				return accessResult;
@@ -974,7 +983,7 @@ namespace Jupiter.Controllers
 			[FromRoute][Required] BucketId bucket,
 			[FromRoute][Required] RefId key)
 		{
-			ActionResult? accessResult = await _requestHelper.HasAccessToNamespaceAsync(User, Request, ns, new[] { JupiterAclAction.WriteObject });
+			ActionResult? accessResult = await _requestHelper.HasAccessToScopeAsync(User, Request, new AccessScope(ns, bucket), new[] { JupiterAclAction.WriteObject });
 			if (accessResult != null)
 			{
 				return accessResult;
@@ -1047,7 +1056,7 @@ namespace Jupiter.Controllers
 			[FromRoute][Required] RefId key,
 			[FromRoute][Required] BlobId hash)
 		{
-			ActionResult? accessResult = await _requestHelper.HasAccessToNamespaceAsync(User, Request, ns, new[] { JupiterAclAction.WriteObject });
+			ActionResult? accessResult = await _requestHelper.HasAccessToScopeAsync(User, Request, new AccessScope(ns, bucket), new[] { JupiterAclAction.WriteObject });
 			if (accessResult != null)
 			{
 				return accessResult;
@@ -1092,10 +1101,10 @@ namespace Jupiter.Controllers
 			[Required] ContentId id,
 			[FromQuery] bool supportsRedirect = false)
 		{
-			ActionResult? result = await _requestHelper.HasAccessToNamespaceAsync(User, Request, ns, new[] { JupiterAclAction.ReadObject });
-			if (result != null)
+			ActionResult? accessResult = await _requestHelper.HasAccessToScopeAsync(User, Request, new AccessScope(ns, bucket), new[] { JupiterAclAction.ReadObject });
+			if (accessResult != null)
 			{
-				return result;
+				return accessResult;
 			}
 
 			try
@@ -1152,10 +1161,10 @@ namespace Jupiter.Controllers
 #pragma warning restore IDE0060
 			[Required] BlobId id)
 		{
-			ActionResult? result = await _requestHelper.HasAccessToNamespaceAsync(User, Request, ns, new[] { JupiterAclAction.WriteObject });
-			if (result != null)
+			ActionResult? accessResult = await _requestHelper.HasAccessToScopeAsync(User, Request, new AccessScope(ns, bucket), new[] { JupiterAclAction.WriteObject });
+			if (accessResult != null)
 			{
-				return result;
+				return accessResult;
 			}
 
 			_diagnosticContext.Set("Content-Length", Request.ContentLength ?? -1);
@@ -1234,12 +1243,13 @@ namespace Jupiter.Controllers
 				}
 			}
 
-			JupiterAclAction[] requiredActions = ops.Ops.Select(op => ActionForOp(op.Op)).ToArray();
-
-			ActionResult? accessResult = await _requestHelper.HasAccessToNamespaceAsync(User, Request, ns, requiredActions);
-			if (accessResult != null)
+			foreach (BatchOps.BatchOp op in ops.Ops)
 			{
-				return accessResult;
+				ActionResult? accessResult = await _requestHelper.HasAccessToScopeAsync(User, Request, new AccessScope(ns, op.Bucket), new[] { ActionForOp(op.Op) });
+				if (accessResult != null)
+				{
+					return accessResult;
+				}
 			}
 
 			Tracer.CurrentSpan.SetAttribute("CountOfOps", ops.Ops.Length);
@@ -1446,7 +1456,7 @@ namespace Jupiter.Controllers
 			[FromRoute][Required] NamespaceId ns
 		)
 		{
-			ActionResult? accessResult = await _requestHelper.HasAccessToNamespaceAsync(User, Request, ns, new[] { JupiterAclAction.DeleteNamespace });
+			ActionResult? accessResult = await _requestHelper.HasAccessToScopeAsync(User, Request, new AccessScope(ns), new[] { JupiterAclAction.DeleteNamespace });
 			if (accessResult != null)
 			{
 				return accessResult;
@@ -1479,7 +1489,7 @@ namespace Jupiter.Controllers
 			[FromRoute][Required] NamespaceId ns,
 			[FromRoute][Required] BucketId bucket)
 		{
-			ActionResult? accessResult = await _requestHelper.HasAccessToNamespaceAsync(User, Request, ns, new[] { JupiterAclAction.DeleteBucket });
+			ActionResult? accessResult = await _requestHelper.HasAccessToScopeAsync(User, Request, new AccessScope(ns, bucket), new[] { JupiterAclAction.DeleteBucket });
 			if (accessResult != null)
 			{
 				return accessResult;
@@ -1516,7 +1526,7 @@ namespace Jupiter.Controllers
 			[FromRoute][Required] BucketId bucket,
 			[FromRoute][Required] RefId key)
 		{
-			ActionResult? accessResult = await _requestHelper.HasAccessToNamespaceAsync(User, Request, ns, new[] { JupiterAclAction.DeleteObject });
+			ActionResult? accessResult = await _requestHelper.HasAccessToScopeAsync(User, Request, new AccessScope(ns, bucket), new[] { JupiterAclAction.DeleteObject });
 			if (accessResult != null)
 			{
 				return accessResult;
diff --git a/Engine/Source/Programs/UnrealCloudDDC/Jupiter/Controllers/ReplicationController.cs b/Engine/Source/Programs/UnrealCloudDDC/Jupiter/Controllers/ReplicationController.cs
--- a/Engine/Source/Programs/UnrealCloudDDC/Jupiter/Controllers/ReplicationController.cs
+++ b/Engine/Source/Programs/UnrealCloudDDC/Jupiter/Controllers/ReplicationController.cs
@@ -34,10 +34,10 @@ namespace Jupiter.Controllers
 			[Required] NamespaceId ns
 		)
 		{
-			ActionResult? result = await _requestHelper.HasAccessToNamespaceAsync(User, Request, ns, new[] { JupiterAclAction.ReadTransactionLog });
-			if (result != null)
+			ActionResult? accessResult = await _requestHelper.HasAccessToScopeAsync(User, Request, new AccessScope(ns), new[] { JupiterAclAction.ReadTransactionLog });
+			if (accessResult != null)
 			{
-				return result;
+				return accessResult;
 			}
 
 			IEnumerable<IReplicator>? replicators = _replicationService.GetReplicators(ns);
@@ -77,10 +77,10 @@ namespace Jupiter.Controllers
 			[Required] long offset
 		)
 		{
-			ActionResult? result = await _requestHelper.HasAccessToNamespaceAsync(User, Request, ns, new[] { JupiterAclAction.WriteTransactionLog });
-			if (result != null)
+			ActionResult? accessResult = await _requestHelper.HasAccessToScopeAsync(User, Request, new AccessScope(ns), new[] { JupiterAclAction.WriteTransactionLog });
+			if (accessResult != null)
 			{
-				return result;
+				return accessResult;
 			}
 
 			IEnumerable<IReplicator>? replicators = _replicationService.GetReplicators(ns);
@@ -115,10 +115,10 @@ namespace Jupiter.Controllers
 			[Required][FromBody] NewReplicationState replicationState
 		)
 		{
-			ActionResult? result = await _requestHelper.HasAccessToNamespaceAsync(User, Request, ns, new[] { JupiterAclAction.WriteTransactionLog });
-			if (result != null)
+			ActionResult? accessResult = await _requestHelper.HasAccessToScopeAsync(User, Request, new AccessScope(ns), new[] { JupiterAclAction.WriteTransactionLog });
+			if (accessResult != null)
 			{
-				return result;
+				return accessResult;
 			}
 
 			IEnumerable<IReplicator>? replicators = _replicationService.GetReplicators(ns);
@@ -158,10 +158,10 @@ namespace Jupiter.Controllers
 			[Required] NamespaceId ns
 		)
 		{
-			ActionResult? result = await _requestHelper.HasAccessToNamespaceAsync(User, Request, ns, new[] { JupiterAclAction.WriteTransactionLog });
-			if (result != null)
+			ActionResult? accessResult = await _requestHelper.HasAccessToScopeAsync(User, Request, new AccessScope(ns), new[] { JupiterAclAction.WriteTransactionLog });
+			if (accessResult != null)
 			{
-				return result;
+				return accessResult;
 			}
 
 			List<IReplicator> replicators = _replicationService.GetReplicators(ns).ToList();
diff --git a/Engine/Source/Programs/UnrealCloudDDC/Jupiter/Controllers/ReplicationLogController.cs b/Engine/Source/Programs/UnrealCloudDDC/Jupiter/Controllers/ReplicationLogController.cs
--- a/Engine/Source/Programs/UnrealCloudDDC/Jupiter/Controllers/ReplicationLogController.cs
+++ b/Engine/Source/Programs/UnrealCloudDDC/Jupiter/Controllers/ReplicationLogController.cs
@@ -46,10 +46,10 @@ namespace Jupiter.Controllers
 			[Required] NamespaceId ns
 		)
 		{
-			ActionResult? result = await _requestHelper.HasAccessToNamespaceAsync(User, Request, ns, new[] { JupiterAclAction.ReadTransactionLog });
-			if (result != null)
+			ActionResult? accessResult = await _requestHelper.HasAccessToScopeAsync(User, Request, new AccessScope(ns), new[] { JupiterAclAction.ReadTransactionLog });
+			if (accessResult != null)
 			{
-				return result;
+				return accessResult;
 			}
 
 			try
@@ -69,10 +69,10 @@ namespace Jupiter.Controllers
 			[Required] NamespaceId ns
 		)
 		{
-			ActionResult? result = await _requestHelper.HasAccessToNamespaceAsync(User, Request, ns, new[] { JupiterAclAction.WriteTransactionLog });
-			if (result != null)
+			ActionResult? accessResult = await _requestHelper.HasAccessToScopeAsync(User, Request, new AccessScope(ns), new[] { JupiterAclAction.WriteTransactionLog });
+			if (accessResult != null)
 			{
-				return result;
+				return accessResult;
 			}
 
 			try
@@ -97,10 +97,10 @@ namespace Jupiter.Controllers
 			[FromQuery] int count = 1000
 		)
 		{
-			ActionResult? result = await _requestHelper.HasAccessToNamespaceAsync(User, Request, ns, new[] { JupiterAclAction.ReadTransactionLog });
-			if (result != null)
+			ActionResult? accessResult = await _requestHelper.HasAccessToScopeAsync(User, Request, new AccessScope(ns), new[] { JupiterAclAction.ReadTransactionLog });
+			if (accessResult != null)
 			{
-				return result;
+				return accessResult;
 			}
 
 			if (((lastBucket == null && lastEvent.HasValue) || (lastBucket != null && !lastEvent.HasValue)) && lastBucket != "now")
@@ -162,10 +162,10 @@ namespace Jupiter.Controllers
 			[Required] string replicationBucket
 		)
 		{
-			ActionResult? result = await _requestHelper.HasAccessToNamespaceAsync(User, Request, ns, new[] { JupiterAclAction.ReadTransactionLog });
-			if (result != null)
+			ActionResult? accessResult = await _requestHelper.HasAccessToScopeAsync(User, Request, new AccessScope(ns), new[] { JupiterAclAction.ReadTransactionLog });
+			if (accessResult != null)
 			{
-				return result;
+				return accessResult;
 			}
 			try
 			{
diff --git a/Engine/Source/Programs/UnrealCloudDDC/Jupiter/Controllers/SymbolsController.cs b/Engine/Source/Programs/UnrealCloudDDC/Jupiter/Controllers/SymbolsController.cs
--- a/Engine/Source/Programs/UnrealCloudDDC/Jupiter/Controllers/SymbolsController.cs
+++ b/Engine/Source/Programs/UnrealCloudDDC/Jupiter/Controllers/SymbolsController.cs
@@ -81,14 +81,14 @@ namespace Jupiter.Controllers
 			[FromRoute][Required] string fileName,
 			[FromQuery] bool supportsRedirect = false)
 		{
-			ActionResult? accessResult = await _requestHelper.HasAccessToNamespaceAsync(User, Request, ns, new[] { JupiterAclAction.ReadObject });
+			BucketId bucket = new BucketId(moduleName);
+			ActionResult? accessResult = await _requestHelper.HasAccessToScopeAsync(User, Request, new AccessScope(ns, bucket), new[] { JupiterAclAction.ReadObject });
 			if (accessResult != null)
 			{
 				return accessResult;
 			}
 
 			LogAuditEntry(HttpMethod.Get, ns, moduleName, identifier, fileName);
-			BucketId bucket = new BucketId(moduleName);
 			BlobContents? refContents;
 			try
 			{
@@ -198,7 +198,8 @@ namespace Jupiter.Controllers
 			[FromRoute][Required] NamespaceId ns,
 			[FromRoute][Required] string moduleName)
 		{
-			ActionResult? accessResult = await _requestHelper.HasAccessToNamespaceAsync(User, Request, ns, new[] { JupiterAclAction.WriteObject });
+			BucketId bucket = new BucketId(moduleName);
+			ActionResult? accessResult = await _requestHelper.HasAccessToScopeAsync(User, Request, new AccessScope(ns, bucket), new[] { JupiterAclAction.WriteObject });
 			if (accessResult != null)
 			{
 				return accessResult;
@@ -220,7 +221,6 @@ namespace Jupiter.Controllers
 				payloadToUse = await _bufferedPayloadFactory.CreateFromStreamAsync(compressedStream, compressedStream.Length, "symbol-put-recompress", HttpContext.RequestAborted);
 			}
 
-			BucketId bucketName = new BucketId(moduleName);
 			using IBufferedPayload payload = payloadToUse;
 			await using Stream hashStream = payload.GetStream();
 			BlobId attachmentHash = await BlobId.FromStreamAsync(hashStream, HttpContext.RequestAborted);
@@ -248,9 +248,9 @@ namespace Jupiter.Controllers
 
 			try
 			{
-				await _blobStore.PutObjectKnownHashAsync(ns, payload, attachmentHash, bucketName, bypassCache: bypassCache, cancellationToken: HttpContext.RequestAborted);
+				await _blobStore.PutObjectKnownHashAsync(ns, payload, attachmentHash, bucket, bypassCache: bypassCache, cancellationToken: HttpContext.RequestAborted);
 
-				(ContentId[], BlobId[]) missingHashes = await _refService.PutAsync(ns, bucketName, RefId.FromName($"{moduleName}.{pdbIdentifier}{pdbAge}.{filename}"), blobHeader, o, cancellationToken: HttpContext.RequestAborted);
+				(ContentId[], BlobId[]) missingHashes = await _refService.PutAsync(ns, bucket, RefId.FromName($"{moduleName}.{pdbIdentifier}{pdbAge}.{filename}"), blobHeader, o, cancellationToken: HttpContext.RequestAborted);
 
 				if (missingHashes.Item1.Any() || missingHashes.Item2.Any())
 				{
diff --git a/Engine/Source/Programs/UnrealCloudDDC/Jupiter/Implementation/Blob/BlobService.cs b/Engine/Source/Programs/UnrealCloudDDC/Jupiter/Implementation/Blob/BlobService.cs
--- a/Engine/Source/Programs/UnrealCloudDDC/Jupiter/Implementation/Blob/BlobService.cs
+++ b/Engine/Source/Programs/UnrealCloudDDC/Jupiter/Implementation/Blob/BlobService.cs
@@ -281,7 +281,8 @@ public class BlobService : IBlobService
 					throw;
 				}
 
-				ActionResult? result = await _requestHelper.HasAccessToNamespaceAsync(user, request, policy.FallbackNamespace.Value, new[] { JupiterAclAction.ReadObject });
+				AccessScope accessScope = bucketHint == null ? new AccessScope(policy.FallbackNamespace) : new AccessScope(policy.FallbackNamespace, bucketHint.Value); 
+				ActionResult? result = await _requestHelper.HasAccessToScopeAsync(user, request, accessScope, new[] { JupiterAclAction.ReadObject });
 				if (result != null)
 				{
 					_logger.LogInformation("Authorization error when attempting to fallback to namespace {FallbackNamespace}. This may be confusing for users that as they had access to original namespace {Namespace}", policy.FallbackNamespace.Value, ns);
diff --git a/Engine/Source/Programs/UnrealCloudDDC/Jupiter/appsettings.Testing.json b/Engine/Source/Programs/UnrealCloudDDC/Jupiter/appsettings.Testing.json
--- a/Engine/Source/Programs/UnrealCloudDDC/Jupiter/appsettings.Testing.json
+++ b/Engine/Source/Programs/UnrealCloudDDC/Jupiter/appsettings.Testing.json
@@ -107,633 +107,136 @@
 						}
 					]
 				}
+			},
+			{
+				"Claims": [ "*" ],
+				"Actions": [
+					"ReadObject",
+					"WriteObject",
+					"DeleteObject"
+				],
+				"Scope": {
+					"Namespaces": [
+						{
+							"values": [
+								"ue4.ddc", "TestNamespace", "testbucket", "fallback-namespace","first-namespace", "second-namespace", "test-blobindex-memory", "test-blobindex-list-memory", "test-blobindex-mongo", "test-blobindex-list-mongo", "test-blobindex-scylla", "test-blobindex-list-scylla", "test-blobindex-cassandra", "test-blobindex-list-cassandra" ]
+						}
+					],
+					"Buckets": [
+						{
+							"values": [ "*" ]
+						}
+					]
+				}
+			},
+			{
+				"Claims": [ "*" ],
+				"Actions": [
+					"ReadObject",
+					"WriteObject",
+					"DeleteObject",
+					"DeleteBucket",
+					"ReadTransactionLog",
+					"WriteTransactionLog",
+					"AdminAction"
+				],
+				"Scope": {
+					"Namespaces": [
+						{
+							"values": [ "test-namespace", "test-namespace-scylla", "test-namespace-mongo", "test-namespace-memory", "test-namespace-memory", "test-namespace-cassandra", "test-namespace-replication", "test-namespace-enumeration", "test-namespace-bundle", "test-namespace-bypass-cache", "test-namespace-redirect", "test-symbols-memory", "test-symbols-scylla" ]
+						}
+					],
+					"Buckets": [
+						{
+							"values": [ "*" ]
+						}
+					]
+				}
+			},
+			{
+				"Claims": [ "*" ],
+				"Actions": [
+					"ReadObject",
+					"WriteObject",
+					"DeleteObject",
+					"DeleteBucket",
+					"DeleteNamespace"
+				],
+				"Scope": {
+					"Namespaces": [
+						{
+							"values": [ "test-delete-namespace", "test-namespace-delete-cassandra", "test-namespace-delete-scylla", "test-namespace-delete-mongo", "test-namespace-delete-memory", "test-namespace-no-overwrite-memory", "test-namespace-no-overwrite-mongo", "test-namespace-no-overwrite-scylla", "test-namespace-no-overwrite-cassandra" ]
+						}
+					],
+					"Buckets": [
+						{
+							"values": [ "*" ]
+						}
+					]
+				}
+			},
+			{
+				// global operations run in all namespaces and need the admin action
+				// this should really be protected with a specific claim but for tests we disable auth so we do not have claims to control on
+				// do NOT copy this setup for production deployments
+				"Claims": [ "*" ],
+				"Actions": [
+					"AdminAction"
+				],
+				"Scope": {
+					"Namespaces": [
+						{
+							// global operations
+							"values": [ "*" ]
+						}
+					],
+					"Buckets": [
+						{
+							"values": [ "*" ]
+						}
+					]
+				}
 			}
 		]
 	},
 	"Namespaces": {
 		"Policies": {
-			"ue4.ddc": {
-				"Acls": [
-					{
-						"Claims": [ "*" ],
-						"Actions": [
-							"ReadObject",
-							"WriteObject"
-						]
-					}
-				]
-			},
-			"TestNamespace": {
-				"Acls": [
-					{
-						"Claims": [ "*" ],
-						"Actions": [
-							"ReadObject",
-							"WriteObject"
-						]
-					}
-				]
-			},
 			"test-namespace": {
-				"Acls": [
-					{
-						"Claims": [ "*" ],
-						"Actions": [
-							"ReadObject",
-							"WriteObject",
-							"DeleteObject",
-							"DeleteBucket",
-							"ReadTransactionLog",
-							"WriteTransactionLog",
-							"AdminAction"
-						]
-					}
-				],
 				"StoragePool": "test"
 			},
-			"test-namespace-scylla": {
-				"Acls": [
-					{
-						"Claims": [ "*" ],
-						"Actions": [
-							"ReadObject",
-							"WriteObject",
-							"DeleteObject",
-							"DeleteBucket",
-							"DeleteNamespace",
-							"ReadTransactionLog",
-							"WriteTransactionLog",
-							"AdminAction"
-						]
-					}
-				]
-			},
-			"test-namespace-mongo": {
-				"Acls": [
-					{
-						"Claims": [ "*" ],
-						"Actions": [
-							"ReadObject",
-							"WriteObject",
-							"DeleteObject",
-							"DeleteBucket",
-							"DeleteNamespace",
-							"ReadTransactionLog",
-							"WriteTransactionLog",
-							"AdminAction"
-						]
-					}
-				]
-			},
-			"test-namespace-memory": {
-				"Acls": [
-					{
-						"Claims": [ "*" ],
-						"Actions": [
-							"ReadObject",
-							"WriteObject",
-							"DeleteObject",
-							"DeleteBucket",
-							"DeleteNamespace",
-							"ReadTransactionLog",
-							"WriteTransactionLog",
-							"AdminAction"
-						]
-					}
-				]
-			},
-			"test-namespace-cassandra": {
-				"Acls": [
-					{
-						"Claims": [ "*" ],
-						"Actions": [
-							"ReadObject",
-							"WriteObject",
-							"DeleteObject",
-							"DeleteBucket",
-							"DeleteNamespace",
-							"ReadTransactionLog",
-							"WriteTransactionLog",
-							"AdminAction"
-						]
-					}
-				]
-			},
-			"test-namespace-replication": {
-				"Acls": [
-					{
-						"Claims": [ "*" ],
-						"Actions": [
-							"ReadObject",
-							"WriteObject",
-							"DeleteObject",
-							"DeleteBucket",
-							"ReadTransactionLog",
-							"WriteTransactionLog",
-							"AdminAction"
-						]
-					}
-				]
-			},
-			"test-namespace-enumeration": {
-				"Acls": [
-					{
-						"Claims": [ "*" ],
-						"Actions": [
-							"ReadObject",
-							"WriteObject",
-							"DeleteObject",
-							"DeleteBucket",
-							"ReadTransactionLog",
-							"WriteTransactionLog",
-							"AdminAction",
-							"EnumerateBucket"
-						]
-					}
-				]
-			},
 			"test-namespace-bundle": {
-				"Acls": [
-					{
-						"Claims": [ "*" ],
-						"Actions": [
-							"ReadObject",
-							"WriteObject",
-							"DeleteObject",
-							"DeleteBucket",
-							"ReadTransactionLog",
-							"WriteTransactionLog",
-							"AdminAction"
-						]
-					}
-				],
 				"UseContentAddressedStorage": false,
 				"AllowRedirectUris": true
 			},
 			"test-namespace-bypass-cache": {
-				"Acls": [
-					{
-						"Claims": [ "*" ],
-						"Actions": [
-							"ReadObject",
-							"WriteObject",
-							"DeleteObject",
-							"DeleteBucket",
-							"ReadTransactionLog",
-							"WriteTransactionLog",
-							"AdminAction"
-						]
-					}
-				],
 				"BypassCacheOnWrite": true
 			},
 			"test-namespace-redirect": {
-				"AllowRedirectUris": true,
-				"Acls": [
-					{
-						"Claims": [ "*" ],
-						"Actions": [
-							"ReadObject",
-							"WriteObject",
-							"DeleteObject",
-							"DeleteBucket",
-							"ReadTransactionLog",
-							"WriteTransactionLog",
-							"AdminAction"
-						]
-					}
-				]
-			},
-			"delete-namespace": {
-				"Acls": [
-					{
-						"Claims": [ "*" ],
-						"Actions": [
-							"ReadObject",
-							"WriteObject"
-						]
-					}
-				]
-			},
-			"test-delete-namespace": {
-				"Acls": [
-					{
-						"Claims": [ "*" ],
-						"Actions": [
-							"ReadObject",
-							"WriteObject",
-							"DeleteObject",
-							"DeleteBucket",
-							"DeleteNamespace"
-						]
-					}
-				]
-			},
-			"test-namespace-delete-cassandra": {
-				"Acls": [
-					{
-						"Claims": [ "*" ],
-						"Actions": [
-							"ReadObject",
-							"WriteObject",
-							"DeleteObject",
-							"DeleteBucket",
-							"DeleteNamespace"
-						]
-					}
-				]
-			},
-			"test-namespace-delete-mongo": {
-				"Acls": [
-					{
-						"Claims": [ "*" ],
-						"Actions": [
-							"ReadObject",
-							"WriteObject",
-							"DeleteObject",
-							"DeleteBucket",
-							"DeleteNamespace"
-						]
-					}
-				]
-			},
-			"test-namespace-delete-scylla": {
-				"Acls": [
-					{
-						"Claims": [ "*" ],
-						"Actions": [
-							"ReadObject",
-							"WriteObject",
-							"DeleteObject",
-							"DeleteBucket",
-							"DeleteNamespace"
-						]
-					}
-				]
-			},
-			"test-namespace-delete-memory": {
-				"Acls": [
-					{
-						"Claims": [ "*" ],
-						"Actions": [
-							"ReadObject",
-							"WriteObject",
-							"DeleteObject",
-							"DeleteBucket",
-							"DeleteNamespace"
-						]
-					}
-				]
+				"AllowRedirectUris": true
 			},
 			"test-namespace-no-overwrite-memory": {
-				"AllowOverwritesOfRefs": false,
-				"Acls": [
-					{
-						"Claims": [ "*" ],
-						"Actions": [
-							"ReadObject",
-							"WriteObject",
-							"DeleteObject",
-							"DeleteBucket",
-							"DeleteNamespace"
-						]
-					}
-				]
+				"AllowOverwritesOfRefs": false
 			},
 			"test-namespace-no-overwrite-mongo": {
-				"AllowOverwritesOfRefs": false,
-				"Acls": [
-					{
-						"Claims": [ "*" ],
-						"Actions": [
-							"ReadObject",
-							"WriteObject",
-							"DeleteObject",
-							"DeleteBucket",
-							"DeleteNamespace"
-						]
-					}
-				]
+				"AllowOverwritesOfRefs": false
 			},
 			"test-namespace-no-overwrite-scylla": {
-				"AllowOverwritesOfRefs": false,
-				"Acls": [
-					{
-						"Claims": [ "*" ],
-						"Actions": [
-							"ReadObject",
-							"WriteObject",
-							"DeleteObject",
-							"DeleteBucket",
-							"DeleteNamespace"
-						]
-					}
-				]
+				"AllowOverwritesOfRefs": false
 			},
 			"test-namespace-no-overwrite-cassandra": {
-				"AllowOverwritesOfRefs": false,
-				"Acls": [
-					{
-						"Claims": [ "*" ],
-						"Actions": [
-							"ReadObject",
-							"WriteObject",
-							"DeleteObject",
-							"DeleteBucket",
-							"DeleteNamespace"
-						]
-					}
-				]
+				"AllowOverwritesOfRefs": false
 			},
 			"test-builds-scylla": {
-				"Acls": [
-					{
-						"Claims": [ "*" ],
-						"Actions": [
-							"ReadObject",
-							"WriteObject",
-							"DeleteObject",
-							"DeleteBucket",
-							"DeleteNamespace",
-							"ReadTransactionLog",
-							"WriteTransactionLog",
-							"AdminAction",
-							"EnumerateBucket"
-						]
-					}
-				],
 				"UseBlobIndexForExists": true
 			},
 			"test-builds-copy-scylla": {
-				"Acls": [
-					{
-						"Claims": [ "*" ],
-						"Actions": [
-							"ReadObject",
-							"WriteObject",
-							"DeleteObject",
-							"DeleteBucket",
-							"DeleteNamespace",
-							"ReadTransactionLog",
-							"WriteTransactionLog",
-							"AdminAction",
-							"EnumerateBucket"
-						]
-					}
-				],
 				"UseBlobIndexForExists":  true
 			},
-			"test-builds-memory": {
-				"Acls": [
-					{
-						"Claims": [ "*" ],
-						"Actions": [
-							"ReadObject",
-							"WriteObject",
-							"DeleteObject",
-							"DeleteBucket",
-							"DeleteNamespace",
-							"ReadTransactionLog",
-							"WriteTransactionLog",
-							"AdminAction",
-							"EnumerateBucket"
-						]
-					}
-				]
-			},
-			"testbucket": {
-				"Acls": [
-					{
-						"Claims": [ "*" ],
-						"Actions": [
-							"ReadObject",
-							"WriteObject",
-							"DeleteObject"
-						]
-					}
-				]
-			},
-			"test-blobindex-memory": {
-				"Acls": [
-					{
-						"Claims": [ "*" ],
-						"Actions": [
-							"ReadObject",
-							"WriteObject",
-							"DeleteObject"
-						]
-					}
-				]
-			},
-			"test-blobindex-list-memory": {
-				"Acls": [
-					{
-						"Claims": [ "*" ],
-						"Actions": [
-							"ReadObject",
-							"WriteObject",
-							"DeleteObject"
-						]
-					}
-				]
-			},
-			"test-blobindex-mongo": {
-				"Acls": [
-					{
-						"Claims": [ "*" ],
-						"Actions": [
-							"ReadObject",
-							"WriteObject",
-							"DeleteObject"
-						]
-					}
-				]
-			},
-			"test-blobindex-list-mongo": {
-				"Acls": [
-					{
-						"Claims": [ "*" ],
-						"Actions": [
-							"ReadObject",
-							"WriteObject",
-							"DeleteObject"
-						]
-					}
-				]
-			},
-			"test-blobindex-cassandra": {
-				"Acls": [
-					{
-						"Claims": [ "*" ],
-						"Actions": [
-							"ReadObject",
-							"WriteObject",
-							"DeleteObject"
-						]
-					}
-				]
-			},
-			"test-blobindex-list-cassandra": {
-				"Acls": [
-					{
-						"Claims": [ "*" ],
-						"Actions": [
-							"ReadObject",
-							"WriteObject",
-							"DeleteObject"
-						]
-					}
-				]
-			},
-			"test-blobindex-scylla": {
-				"Acls": [
-					{
-						"Claims": [ "*" ],
-						"Actions": [
-							"ReadObject",
-							"WriteObject",
-							"DeleteObject"
-						]
-					}
-				]
-			},
-			"test-blobindex-list-scylla": {
-				"Acls": [
-					{
-						"Claims": [ "*" ],
-						"Actions": [
-							"ReadObject",
-							"WriteObject",
-							"DeleteObject"
-						]
-					}
-				]
-			},
-			"test-namespace-last-access": {
-				"Acls": [
-					{
-						"Claims": [ "*" ],
-						"Actions": [
-							"ReadObject",
-							"WriteObject"
-						]
-					}
-				]
-			},
-			"test-namespace-gcref": {
-				"Acls": [
-					{
-						"Claims": [ "*" ],
-						"Actions": [
-							"ReadObject",
-							"WriteObject",
-							"DeleteObject",
-							"DeleteBucket",
-							"ReadTransactionLog",
-							"WriteTransactionLog",
-							"AdminAction"
-						]
-					}
-				]
-			},
-			"test-namespace-gcref-sharded": {
-				"Acls": [
-					{
-						"Claims": [ "*" ],
-						"Actions": [
-							"ReadObject",
-							"WriteObject",
-							"DeleteObject",
-							"DeleteBucket",
-							"ReadTransactionLog",
-							"WriteTransactionLog",
-							"AdminAction"
-						]
-					}
-				]
-			},
-			"test-symbols-memory": {
-				"Acls": [
-					{
-						"Claims": [ "*" ],
-						"Actions": [
-							"ReadObject",
-							"WriteObject",
-							"DeleteObject",
-							"DeleteBucket",
-							"ReadTransactionLog",
-							"WriteTransactionLog",
-							"AdminAction"
-						]
-					}
-				]
-			},
-			"test-symbols-scylla": {
-				"Acls": [
-					{
-						"Claims": [ "*" ],
-						"Actions": [
-							"ReadObject",
-							"WriteObject",
-							"DeleteObject",
-							"DeleteBucket",
-							"ReadTransactionLog",
-							"WriteTransactionLog",
-							"AdminAction"
-						]
-					}
-				]
-			},
-			"snapshot-namespace": {
-				"Acls": [
-					{
-						"Claims": [ "*" ],
-						"Actions": [
-							"ReadObject",
-							"WriteObject"
-						]
-					}
-				]
-			},
 			"first-namespace": {
-				"Acls": [
-					{
-						"Claims": [ "*" ],
-						"Actions": [
-							"ReadObject",
-							"WriteObject"
-						]
-					}
-				],
 				"FallbackNamespace": "fallback-namespace"
 			},
-			"fallback-namespace": {
-				"Acls": [
-					{
-						"Claims": [ "*" ],
-						"Actions": [
-							"ReadObject",
-							"WriteObject"
-						]
-					}
-				]
-			},
 			"second-namespace": {
-				"Acls": [
-					{
-						"Claims": [ "*" ],
-						"Actions": [
-							"ReadObject",
-							"WriteObject"
-						]
-					}
-				],
 				"FallbackNamespace": "fallback-namespace-no-auth"
-			},
-			"fallback-namespace-no-access": {
-				// no acls mapped
-				"Acls": {
-				}
 			}
 		}
 	},
