Cloud DDC - Return 503s when Scylla is unavailable in BuildsController

#Synced-CL 48106205

diff --git a/Engine/Source/Programs/UnrealCloudDDC/Jupiter/Controllers/BuildsController.cs b/Engine/Source/Programs/UnrealCloudDDC/Jupiter/Controllers/BuildsController.cs
--- a/Engine/Source/Programs/UnrealCloudDDC/Jupiter/Controllers/BuildsController.cs
+++ b/Engine/Source/Programs/UnrealCloudDDC/Jupiter/Controllers/BuildsController.cs
@@ -931,6 +931,10 @@ namespace Jupiter.Controllers
 			{
 				return NotFound(new ProblemDetails { Title = $"Blob {e.Blob} in {e.Ns} not found" });
 			}
+			catch (RemoteServiceException e)
+			{
+				return Problem(e.Message, null, (int)HttpStatusCode.ServiceUnavailable);
+			}
 		}
 
 		[HttpPut("{ns}/{bucket}/{buildId}/parts/{partId}/{partName}.{format?}")]
@@ -1029,6 +1033,10 @@ namespace Jupiter.Controllers
 			{
 				return Problem($"Build already exists {ns} {bucket} {buildId}", statusCode: (int)HttpStatusCode.Conflict);
 			}
+			catch (RemoteServiceException e)
+			{
+				return Problem(e.Message, null, (int)HttpStatusCode.ServiceUnavailable);
+			}
 		}
 
 		private async Task<List<ContentHash>> PutBuildPartRefAsync(NamespaceId ns, BucketId bucket, CbObjectId buildId, CbObjectId partId, BlockContext blockContext, CbObject payloadObject, BlobId payloadHash)
@@ -1163,6 +1171,10 @@ namespace Jupiter.Controllers
 			{
 				return NotFound(e.Message);
 			}
+			catch (RemoteServiceException e)
+			{
+				return Problem(e.Message, null, (int)HttpStatusCode.ServiceUnavailable);
+			}
 		}
 
 		private async Task<BuildRecord?> GetBuildWithRetriesAsync(NamespaceId ns, BucketId bucket, CbObjectId buildId)
@@ -1259,6 +1271,19 @@ namespace Jupiter.Controllers
 			{
 				return NotFound(e.Message);
 			}
+			catch (AggregateException e)
+			{
+				if (e.InnerExceptions.Any(exception => exception.GetType() == typeof(RemoteServiceException)))
+				{
+					return Problem(e.Message, null, (int)HttpStatusCode.ServiceUnavailable);
+				}
+
+				throw;
+			}
+			catch (RemoteServiceException e)
+			{
+				return Problem(e.Message, null, (int)HttpStatusCode.ServiceUnavailable);
+			}
 
 			return NotFound();
 		}
@@ -1367,6 +1392,10 @@ namespace Jupiter.Controllers
 			{
 				return NotFound(new ValidationProblemDetails { Title = $"Content Id {e.ContentId} not found" });
 			}
+			catch (RemoteServiceException e)
+			{
+				return Problem(e.Message, null, (int)HttpStatusCode.ServiceUnavailable);
+			}
 		}
 
 		[HttpPut("{ns}/{bucket}/{buildId}/parts/{partId}/blobs/{id}.{format?}")]
@@ -1436,6 +1465,10 @@ namespace Jupiter.Controllers
 			{
 				return StatusCode(StatusCodes.Status408RequestTimeout, "Request cancelled");
 			}
+			catch (RemoteServiceException e)
+			{
+				return Problem(e.Message, null, (int)HttpStatusCode.ServiceUnavailable);
+			}
 		}
 
 		[HttpPost("{ns}/{bucket}/{buildId}/parts/{partId}/blobs/{id}/startMultipartUpload")]
