Cloud DDC - Lowered size of blob required to buffer from S3 when multipart downloading. This uses more temporary diskspace but less memory. Has the benefit of triggering a optimization when populating the filesystem cache were we can just move the data into the cache more often.

#Synced-CL 47492933

diff --git a/Engine/Source/Programs/UnrealCloudDDC/Jupiter/Settings.cs b/Engine/Source/Programs/UnrealCloudDDC/Jupiter/Settings.cs
--- a/Engine/Source/Programs/UnrealCloudDDC/Jupiter/Settings.cs
+++ b/Engine/Source/Programs/UnrealCloudDDC/Jupiter/Settings.cs
@@ -245,7 +245,7 @@ namespace Jupiter
 		/// Max size of bytes that is buffered into memory before switching to filesystem buffers
 		/// Larger value means more RAM used but less load on the filesystem.
 		/// </summary>
-		public int MultiPartMaxMemoryBufferSize { get; set; } = int.MaxValue;
+		public int MultiPartMaxMemoryBufferSize { get; set; } = 67108864; // 64MB or 2 chunks are buffered into memory then we switch to filesystem
 
 		/// <summary>
 		/// Set this to disable payload signing and checksum generation	 during S3 uploads, saves cpu load when uploading large objects at some consistency risk and a requirement to use https connections see
