Cloud DDC - Prepared release 1.4.0

#Synced-CL 48106205

diff --git a/Engine/Source/Programs/UnrealCloudDDC/Helm/UnrealCloudDDC/version.yaml b/Engine/Source/Programs/UnrealCloudDDC/Helm/UnrealCloudDDC/version.yaml
--- a/Engine/Source/Programs/UnrealCloudDDC/Helm/UnrealCloudDDC/version.yaml
+++ b/Engine/Source/Programs/UnrealCloudDDC/Helm/UnrealCloudDDC/version.yaml
@@ -1 +1 @@
-version: 1.3.0
\ No newline at end of file
+version: 1.4.0
\ No newline at end of file
diff --git a/Engine/Source/Programs/UnrealCloudDDC/Jupiter/appsettings.Testing.json b/Engine/Source/Programs/UnrealCloudDDC/Jupiter/appsettings.Testing.json
--- a/Engine/Source/Programs/UnrealCloudDDC/Jupiter/appsettings.Testing.json
+++ b/Engine/Source/Programs/UnrealCloudDDC/Jupiter/appsettings.Testing.json
@@ -28,7 +28,6 @@
 
 	"UnrealCloudDDC": {
 		"StorageImplementations": [ "memory" ],
-		"RequirePrivatePortForEnumeration": false,
 		"EnableForceSubmitRefBlobToBlobStore": true
 	},
 
diff --git a/Engine/Source/Programs/UnrealCloudDDC/README.md b/Engine/Source/Programs/UnrealCloudDDC/README.md
--- a/Engine/Source/Programs/UnrealCloudDDC/README.md
+++ b/Engine/Source/Programs/UnrealCloudDDC/README.md
@@ -25,7 +25,6 @@ Unreal Cloud DDC can signficantly help teams speed up their cook processes in th
   - [Health Checks](#health-checks)
 - [Authentication](#authentication)
   - [IdentityProvider setup](#identityprovider-setup)
-  - [Namespace access (legacy setup)](#namespace-access-legacy-setup)
   - [ACL Policies](#acl-policies)
 - [Common Configuration](#common-configuration)
   - [DDC](#ddc)
@@ -154,72 +153,7 @@ We recommend naming your scheme `Bearer` if its your first and only scheme. You
 
 The implementation field is usually `JWTBearer` but we do offer a `Okta` if you are using Okta with custom auth servers, for Okta using the org auth server you will need to use `JWTBearer` as well.
 
-## Namespace access (legacy setup)
-
-This section refers to the legacy way of setting up access control whithin Cloud DDC which only supports access control for global actions and namespaces (not buckets). Some APIs still use this but its being phased out.
-
-Access to operations within UnrealCloudDDC is controlled using a set of actions:
-```
-    - ReadObject
-    - WriteObject
-    - DeleteObject
-    - DeleteBucket
-    - DeleteNamespace
-    - ReadTransactionLog
-    - WriteTransactionLog
-    - AdminAction
-```
-
-These can be assigned either per namespace using the acls in each namespace policy or by assigning them to the acls in the Auth settings (which applies them across all namespaces and for operations not associated with a namespace)
-Example configuration that sets transaction log access for users that have access to it, admin access for admins and then per namespace access.
-```
-auth:
-    acls:
-    - claims: 
-        - groups=app-ddc-storage-transactionlog
-        actions:
-        - ReadTransactionLog
-        - WriteTransactionLog
-
-    - claims: 
-        - groups=app-ddc-storage-admin
-        actions:
-        - ReadObject
-        - WriteObject
-        - DeleteObject
-        - DeleteBucket
-        - DeleteNamespace
-        - AdminAction
-
-namespace:
-  policies:
-    example-namespace:
-      acls:
-      - actions: 
-        - ReadObject
-        - WriteObject
-        claims: 
-        - ExampleClaim
-     - actions: 
-        - ReadObject
-        - WriteObject
-        claims: 
-        - AnotherClaim
-    open-namespace:
-      acls:
-      - actions: 
-        - ReadObject
-        - WriteObject
-        claims: 
-        - "*"
-```
-
-If you specify multiple claims in the claims array these are ANDed together thus requiring all the claims. A claim statement like `A=B` requires claim `A` to have value `B` (or contain value `B` in the case of a array).
-You can also specify the `*` claim which grants access for any valid token no matter which claims it has, this is mostly for debug / testing scenarios and shouldnt be used for production data.
-
 ## ACL Policies
-This section refers to the new way to setup access control that allows for access control on buckets for the apis that support this.
-
 Access to operations within UnrealCloudDDC is controlled using a set of actions:
 ```
     - ReadObject
diff --git a/Engine/Source/Programs/UnrealCloudDDC/changelog.md b/Engine/Source/Programs/UnrealCloudDDC/changelog.md
--- a/Engine/Source/Programs/UnrealCloudDDC/changelog.md
+++ b/Engine/Source/Programs/UnrealCloudDDC/changelog.md
@@ -1,7 +1,60 @@
 # Unreleased
-* Changed default value of `UnrealCloudDDC:EnableBucketStatsTracking` to `false`, as enabling it can lead to performance issues when querying sizes for many attachments.
+
+# 1.4.0
+Noteworthy and potentially breaking changes:
+* All endpoints have been updated to use the new per-bucket authorization by default. Set the `Auth.UseLegacyConfiguration` setting to `true` to enable the legacy authentication behavior.
+* Namespace-based authorization is deprecated, along with the `Auth.UseLegacyConfiguration` setting, and will be removed in a future release.
+* The `ConnectionString` setting for replicators is no longer strictly required, allowing a `SourceSite` to be defined instead, which will look up the connection string of the relevant peer in the `cluster.peers` configuration. `SourceSite` will be respected over a `ConnectionString`, and an error will be thrown if neither are provided.
+* The `ReferenceResolver.UseV2ReferenceResolving` setting is now enabled by default.
+* The `UnrealCloudDDC.EnableForceSubmitRefBlobToBlobStore` setting is now disabled by default. When disabled, inlined blobs are only written to Scylla, significantly reducing the number of small blobs in S3. When `UnrealCloudDDC.EnableInlineSmallBlobs` is disabled (not recommended), then blobs are still written to S3.
+* The `UnrealCloudDDC.EnableBucketStatsTracking` setting is now disabled by default, as it can lead to performance issues when querying sizes for many attachments. The strategy for collecting these stats will be reworked in a future release.
 * Removed bucket enumeration endpoint (`GET /api/v1/refs/{ns}/{bucket}`) as its performance was inefficient and discouraged, and did not function correctly when `EnableBucketStatsTracking` was set to `false`.
-* Removed configuration to require a private port for bucket enumeration (`UnrealCloudDDC:RequirePrivatePortForEnumeration`) since its relevant endpoint has been removed.
+* Removed the `UnrealCloudDDC.RequirePrivatePortForEnumeration` setting since its relevant endpoint has been removed.
+* Regular expressions passed to the build search endpoint (`POST /api/v2/builds/{ns}/search`) now ignore the case of the provided `bucketRegex`.
+* The `S3.MultiPartMaxMemoryBufferSize` setting has been reduced to `67108864` (was `int.MaxValue`).
+* Added a service to periodically clean up temp files, controlled using the `GC.CleanTempFiles` setting (defaults to `true`).
+
+New features:
+* Added the `Scylla.UseConsistentContentIdResolve` setting, which uses local_quorum consistency when resolving content IDs.
+* Added the `GET /api/v2/builds/{ns}/{bucket}/{buildId}/parts/{partId}/fileList` endpoint, useful for inspecting which files belong to a build.
+* Added the `includePartial` option to `POST /api/v2/builds/{ns}/{bucket}/search` endpoint, used to include non-finalized builds to search results.
+* Added the `S3.SiteBucketsPerStoragePool` setting, useful for defining which site-specific S3 bucket to use for each storage pool. 
+* Added the `SourceSite` setting to each replicator, used to look up the connection URL using a peer by its `Name` in the `cluster.peers` configuration, rather than defining its `ConnectionString` explicitly.
+* Added the `S3.DisablePayloadChecksums` setting, which reduces CPU usage at the cost of payload verification.
+* Added the `GC.AllowBlockMetadataCleanup` setting (defaults to `true`).
+* Added the `UnrealCloudDDC.RemoveReferencesOnBlobDeletion` setting (defaults to `true`).
+* Added the `UnrealCloudDDCOrphanBlobCutoffDurationMinutes` setting (defaults to `60`).
+
+Minor improvements:
+* Added retries for Amazon S3 multipart uploads.
+* Return 408 status codes when blob upload tasks are cancelled.
+* Improved Amazon S3 rate limit handling.
+* Improved blob replication stability.
+* Return 503 status codes when the underlying database is unavailable.
+* Added retries to build existence checks.
+* Added multipart support for blob replication.
+* Improved audit logging parallelism in `BuildsController`.
+* Added string-to-number conversions for build search API when using compact binary objects.
+* Increased the default maximum number of builds scanned during a search from `2000` to `20000`.
+* Validate that all parts of a build are finalized before the build itself can be finalized.
+* Use check-and-set pattern for blob replication state to allow state to be mutated externally.
+* Added optimization to move files that are already buffered on the file system to their final destination, rather than copying the file by reading into an in-memory buffer.
+
+Bug fixes:
+* Fallback to file system buffering when unable to allocate enough memory, even for small payloads.
+* Fixed ObjectId JSON serialization.
+* Fixed retry logic for blob replication.
+* Fixed an issue where last access background service was running more frequently than intended.
+* Fixed NPE in `BuildsController` audit logging when no part ID or build ID was specified, but a namespace was.
+* Removed redundant telemetry tracing scope in blob replication.
+* Fixed an issue where setting `EnableInlineSmallBlobs` to `false` would incorrectly attempt to inline all blobs.
+* Fixed incorrect dependency injection type in `ReplicationController`.
+
+Helm chart changes:
+* Dependency on `bitnami/common` is now vendored into the source tree.
+* Changed the `HorizontalPodAutoscaler` to use container-specific metrics for CPU and memory rather than Pod-wide metrics.
+* Fixed ConfigMap generation when omitting any `Jupiter` configuration.
+* Added missing replicator version configuration to ConfigMap.
 
 # 1.3.0
 Noteworthy and potentially breaking changes:
diff --git a/Engine/Source/Programs/UnrealCloudDDC/version.yaml b/Engine/Source/Programs/UnrealCloudDDC/version.yaml
--- a/Engine/Source/Programs/UnrealCloudDDC/version.yaml
+++ b/Engine/Source/Programs/UnrealCloudDDC/version.yaml
@@ -1 +1 @@
-version: 1.3.0
\ No newline at end of file
+version: 1.4.0
\ No newline at end of file
