Cloud DDC - Revert the lower MultiPartMaxMemoryBufferSize after experiencing disk throughput issues

#Synced-CL 48106205

diff --git a/Engine/Source/Programs/UnrealCloudDDC/Jupiter/Settings.cs b/Engine/Source/Programs/UnrealCloudDDC/Jupiter/Settings.cs
--- a/Engine/Source/Programs/UnrealCloudDDC/Jupiter/Settings.cs
+++ b/Engine/Source/Programs/UnrealCloudDDC/Jupiter/Settings.cs
@@ -245,7 +245,7 @@ namespace Jupiter
 		/// Max size of bytes that is buffered into memory before switching to filesystem buffers
 		/// Larger value means more RAM used but less load on the filesystem.
 		/// </summary>
-		public int MultiPartMaxMemoryBufferSize { get; set; } = 67108864; // 64MB or 2 chunks are buffered into memory then we switch to filesystem
+		public int MultiPartMaxMemoryBufferSize { get; set; } = int.MaxValue;
 
 		/// <summary>
 		/// Set this to disable payload signing and checksum generation	 during S3 uploads, saves cpu load when uploading large objects at some consistency risk and a requirement to use https connections see
