Cloud DDC - Upgrade nginx to 1.29.3

#Synced-CL 48106205

diff --git a/Engine/Source/Programs/UnrealCloudDDC/Helm/UnrealCloudDDC/templates/deployment.yaml b/Engine/Source/Programs/UnrealCloudDDC/Helm/UnrealCloudDDC/templates/deployment.yaml
--- a/Engine/Source/Programs/UnrealCloudDDC/Helm/UnrealCloudDDC/templates/deployment.yaml
+++ b/Engine/Source/Programs/UnrealCloudDDC/Helm/UnrealCloudDDC/templates/deployment.yaml
@@ -207,7 +207,7 @@ spec:
     {{- end }}
     {{- if .Values.nginx.enabled }}
       - name: nginx
-        image: nginx:1.25.4
+        image: nginx:1.29.3
         {{- if .Values.nginx.env }}
         env:
         {{- include "common.tplvalues.render" ( dict "value" .Values.nginx.env "context" $ ) | nindent 8 }}
