Cloud DDC - Throw RemoteServiceException when failing to retrieve buckets from Scylla

#Synced-CL 48106205

diff --git a/Engine/Source/Programs/UnrealCloudDDC/Jupiter/Implementation/Objects/ScyllaReferencesStore.cs b/Engine/Source/Programs/UnrealCloudDDC/Jupiter/Implementation/Objects/ScyllaReferencesStore.cs
--- a/Engine/Source/Programs/UnrealCloudDDC/Jupiter/Implementation/Objects/ScyllaReferencesStore.cs
+++ b/Engine/Source/Programs/UnrealCloudDDC/Jupiter/Implementation/Objects/ScyllaReferencesStore.cs
@@ -571,7 +571,18 @@ namespace Jupiter.Implementation
 
 		public async IAsyncEnumerable<BucketId> GetBucketsAsync(NamespaceId ns, [EnumeratorCancellation] CancellationToken cancellationToken)
 		{
-			foreach (ScyllaBucket scyllaBucket in await _mapper.FetchAsync<ScyllaBucket>("WHERE namespace=?", ns.ToString()))
+			IEnumerable<ScyllaBucket> buckets;
+			
+			try
+			{
+				buckets = await _mapper.FetchAsync<ScyllaBucket>("WHERE namespace=?", ns.ToString());
+			}
+			catch (QueryExecutionException e)
+			{
+				throw new RemoteServiceException("Error communicating with Scylla", e);
+			}
+			
+			foreach (ScyllaBucket scyllaBucket in buckets)
 			{
 				if (scyllaBucket.Bucket == null)
 				{
