Cloud DDC - Add option to BuildStoreConsistencyCheckService to check for missing blobs for all builds and parts

#Synced-CL 48106205

diff --git a/Engine/Source/Programs/UnrealCloudDDC/Jupiter/Implementation/Builds/BuildStoreConsistencyCheckService.cs b/Engine/Source/Programs/UnrealCloudDDC/Jupiter/Implementation/Builds/BuildStoreConsistencyCheckService.cs
--- a/Engine/Source/Programs/UnrealCloudDDC/Jupiter/Implementation/Builds/BuildStoreConsistencyCheckService.cs
+++ b/Engine/Source/Programs/UnrealCloudDDC/Jupiter/Implementation/Builds/BuildStoreConsistencyCheckService.cs
@@ -1,6 +1,7 @@
 // Copyright Epic Games, Inc. All Rights Reserved.
 
 using System;
+using System.Collections.Generic;
 using System.Threading;
 using System.Threading.Tasks;
 using EpicGames.Horde.Storage;
@@ -21,6 +22,10 @@ namespace Jupiter.Implementation
 		private readonly IOptionsMonitor<ConsistencyCheckSettings> _settings;
 		private readonly ILeaderElection _leaderElection;
 		private readonly IBuildStore _buildStore;
+		private readonly IBlobStore _blobStore;
+		private readonly IBlockStore _blockStore;
+		private readonly IRefService _refService;
+		private readonly IReferenceResolver _referenceResolver;
 		private readonly Tracer _tracer;
 		private readonly ILogger _logger;
 
@@ -29,11 +34,15 @@ namespace Jupiter.Implementation
 			return _settings.CurrentValue.EnableBuildStoreConsistencyCheck;
 		}
 
-		public BuildStoreConsistencyCheckService(IOptionsMonitor<ConsistencyCheckSettings> settings, ILeaderElection leaderElection, IBuildStore buildStore, Tracer tracer, ILogger<BuildStoreConsistencyCheckService> logger) : base(serviceName: nameof(BuildStoreConsistencyCheckService), TimeSpan.FromSeconds(settings.CurrentValue.ConsistencyCheckPollFrequencySeconds), new BuildStoreConsistencyState(), logger)
+		public BuildStoreConsistencyCheckService(IOptionsMonitor<ConsistencyCheckSettings> settings, ILeaderElection leaderElection, IBuildStore buildStore, IBlobStore blobStore, IBlockStore blockStore, IRefService refService, IReferenceResolver referenceResolver, Tracer tracer, ILogger<BuildStoreConsistencyCheckService> logger) : base(serviceName: nameof(BuildStoreConsistencyCheckService), TimeSpan.FromSeconds(settings.CurrentValue.ConsistencyCheckPollFrequencySeconds), new BuildStoreConsistencyState(), logger)
 		{
 			_settings = settings;
 			_leaderElection = leaderElection;
 			_buildStore = buildStore;
+			_blobStore = blobStore;
+			_blockStore = blockStore;
+			_refService = refService;
+			_referenceResolver = referenceResolver;
 			_tracer = tracer;
 			_logger = logger;
 		}
@@ -62,10 +71,15 @@ namespace Jupiter.Implementation
 		{
 			ulong countOfBuildsChecked = 0;
 			ulong countOfIncorrectBuildsFound = 0;
+			ulong countOfUnresolvableBlobs = 0;
+			ulong countOfBlobsChecked = 0;
+			ulong countOfMissingBlobsFound = 0;
+			ulong countOfBlocksChecked = 0;
+			ulong countOfMissingBlocksFound = 0;
 			await Parallel.ForEachAsync(_buildStore.ListAllBuildsAsync(cancellationToken), new ParallelOptions
 			{
 				CancellationToken = cancellationToken,
-				MaxDegreeOfParallelism = _settings.CurrentValue.BlobIndexMaxParallelOperations,
+				MaxDegreeOfParallelism = _settings.CurrentValue.BuildStoreConsistencyMaxParallelBuilds,
 			},
 				async (tuple, token) =>
 				{
@@ -79,21 +93,84 @@ namespace Jupiter.Implementation
 
 					using TelemetrySpan scope = _tracer.StartActiveSpan("consistency_check.build_store").SetAttribute("resource.name", $"{ns}.{bucket}.{buildId}").SetAttribute("operation.name", "consistency_check.build_store");
 
-					BuildRecord? record = await _buildStore.GetBuildAsync(ns, bucket, buildId);
-					if (record == null)
+					if (_settings.CurrentValue.CheckBuildStorePartialObjects)
 					{
-						_logger.LogWarning("Partial build {BuildId} in bucket {Bucket} and namespace {Namespace} .", buildId, bucket, ns);
+						BuildRecord? record = await _buildStore.GetBuildAsync(ns, bucket, buildId);
+						if (record == null)
+						{
+							_logger.LogWarning("Partial build {BuildId} in bucket {Bucket} and namespace {Namespace} .", buildId, bucket, ns);
+
+							Interlocked.Increment(ref countOfIncorrectBuildsFound);
+							if (_settings.CurrentValue.AllowDeletesInBuildStore)
+							{
+								await _buildStore.DeleteBuild(ns, bucket, buildId);
+							}
+
+							return;
+						}
+					}
 
-						Interlocked.Increment(ref countOfIncorrectBuildsFound);
-						if (_settings.CurrentValue.AllowDeletesInBuildStore)
+					if (_settings.CurrentValue.CheckBuildStoreMissingBlobs)
+					{
+						await foreach ((string _, CbObjectId partId) in _buildStore.GetBuildPartsAsync(ns, bucket, buildId).WithCancellation(token))
 						{
-							await _buildStore.DeleteBuild(ns, bucket, buildId);
+							RefId refId = RefId.FromName($"{buildId}/{partId}");
+							List<BlobId> referencedBlobs = await _refService.GetReferencedBlobsAsync(ns, bucket, refId, ignoreMissingBlobs: true, cancellationToken: token);
+
+							await Parallel.ForEachAsync(referencedBlobs, new ParallelOptions
+							{
+								CancellationToken = token,
+								MaxDegreeOfParallelism = _settings.CurrentValue.BuildStoreConsistencyMaxParallelBlobs
+							},
+								async (referencedBlobId, innerToken) =>
+								{
+									BlobId[]? blobIds = await _referenceResolver.ResolveIdAsync(ns, referencedBlobId, cancellationToken: innerToken);
+									if (blobIds == null)
+									{
+										_logger.LogWarning("Unable to resolve blobId for referenced blob {ReferencedBlobId} in {Namespace}.{Bucket}.{BuildId}.{PartId}.", referencedBlobId, ns, bucket, buildId, partId);
+										Interlocked.Increment(ref countOfUnresolvableBlobs);
+										return;
+									}
+
+									foreach (BlobId blobId in blobIds)
+									{
+										Interlocked.Increment(ref countOfBlobsChecked);
+
+										if (!await _blobStore.ExistsAsync(ns, blobId, cancellationToken: innerToken))
+										{
+											_logger.LogWarning("Found missing blob {BlobId} in {Namespace}.{Bucket}.{BuildId}.{PartId}.", blobId, ns, bucket, buildId, partId);
+											Interlocked.Increment(ref countOfMissingBlobsFound);
+										}
+
+										BlobId? metadataHash = await _blockStore.GetBlockMetadataAsync(ns, blobId);
+
+										if (metadataHash != null)
+										{
+											Interlocked.Increment(ref countOfBlocksChecked);
+											
+											if (!await _blobStore.ExistsAsync(ns, metadataHash, cancellationToken: innerToken))
+											{
+												_logger.LogWarning("Found missing block metadata {MetadataHash} for blob {BlobId} in {Namespace}.{Bucket}.{BuildId}.{PartId}.", metadataHash, blobId, ns, bucket, buildId, partId);
+												Interlocked.Increment(ref countOfMissingBlocksFound);
+											}
+										}
+									}
+								}
+							);
 						}
 					}
 				}
 			);
 
-			_logger.LogInformation("Build Store Consistency check finished, found {CountOfIncorrectBuilds} incorrect builds. Processed {CountOfBuilds} builds.", countOfIncorrectBuildsFound, countOfBuildsChecked);
+			if (_settings.CurrentValue.CheckBuildStorePartialObjects)
+			{
+				_logger.LogInformation("Build Store Consistency check finished, found {CountOfIncorrectBuilds} incorrect builds. Processed {CountOfBuilds} builds.", countOfIncorrectBuildsFound, countOfBuildsChecked);
+			}
+			
+			if (_settings.CurrentValue.CheckBuildStoreMissingBlobs)
+			{
+				_logger.LogInformation("Build Store Consistency check finished. Found {CountOfUnresolvableBlobs} unresolvable blobs. Found {CountOfMissingBlobsFound} missing blobs ({CountOfBlobsChecked} total). Found {CountOfMissingBlocksFound} missing blocks ({CountOfBlocksChecked} total). Processed {CountOfBuilds} builds.", countOfUnresolvableBlobs, countOfMissingBlobsFound, countOfBlobsChecked, countOfMissingBlocksFound, countOfBlocksChecked, countOfBuildsChecked);
+			}
 		}
 
 		protected override Task OnStopping(BuildStoreConsistencyState state)
diff --git a/Engine/Source/Programs/UnrealCloudDDC/Jupiter/Settings.cs b/Engine/Source/Programs/UnrealCloudDDC/Jupiter/Settings.cs
--- a/Engine/Source/Programs/UnrealCloudDDC/Jupiter/Settings.cs
+++ b/Engine/Source/Programs/UnrealCloudDDC/Jupiter/Settings.cs
@@ -327,5 +327,9 @@ namespace Jupiter
 		public bool EnableBlobReferenceChecks { get; set; } = false;
 		public bool EnableBuildStoreConsistencyCheck { get; set; } = false;
 		public bool AllowDeletesInBuildStore { get; set; } = false;
+		public bool CheckBuildStorePartialObjects { get; set; } = false;
+		public bool CheckBuildStoreMissingBlobs { get; set; } = false;
+		public int BuildStoreConsistencyMaxParallelBuilds { get; set; } = 4;
+		public int BuildStoreConsistencyMaxParallelBlobs { get; set; } = 4;
 	}
 }
