Cloud DDC - Prefer using IBlobService over IBlobStore for existence checks

#Synced-CL 48106205

diff --git a/Engine/Source/Programs/UnrealCloudDDC/Jupiter/Implementation/Builds/BuildStoreConsistencyCheckService.cs b/Engine/Source/Programs/UnrealCloudDDC/Jupiter/Implementation/Builds/BuildStoreConsistencyCheckService.cs
--- a/Engine/Source/Programs/UnrealCloudDDC/Jupiter/Implementation/Builds/BuildStoreConsistencyCheckService.cs
+++ b/Engine/Source/Programs/UnrealCloudDDC/Jupiter/Implementation/Builds/BuildStoreConsistencyCheckService.cs
@@ -22,7 +22,7 @@ namespace Jupiter.Implementation
 		private readonly IOptionsMonitor<ConsistencyCheckSettings> _settings;
 		private readonly ILeaderElection _leaderElection;
 		private readonly IBuildStore _buildStore;
-		private readonly IBlobStore _blobStore;
+		private readonly IBlobService _blobService;
 		private readonly IBlockStore _blockStore;
 		private readonly IRefService _refService;
 		private readonly IReferenceResolver _referenceResolver;
@@ -34,12 +34,12 @@ namespace Jupiter.Implementation
 			return _settings.CurrentValue.EnableBuildStoreConsistencyCheck;
 		}
 
-		public BuildStoreConsistencyCheckService(IOptionsMonitor<ConsistencyCheckSettings> settings, ILeaderElection leaderElection, IBuildStore buildStore, IBlobStore blobStore, IBlockStore blockStore, IRefService refService, IReferenceResolver referenceResolver, Tracer tracer, ILogger<BuildStoreConsistencyCheckService> logger) : base(serviceName: nameof(BuildStoreConsistencyCheckService), TimeSpan.FromSeconds(settings.CurrentValue.ConsistencyCheckPollFrequencySeconds), new BuildStoreConsistencyState(), logger)
+		public BuildStoreConsistencyCheckService(IOptionsMonitor<ConsistencyCheckSettings> settings, ILeaderElection leaderElection, IBuildStore buildStore, IBlobService blobService, IBlockStore blockStore, IRefService refService, IReferenceResolver referenceResolver, Tracer tracer, ILogger<BuildStoreConsistencyCheckService> logger) : base(serviceName: nameof(BuildStoreConsistencyCheckService), TimeSpan.FromSeconds(settings.CurrentValue.ConsistencyCheckPollFrequencySeconds), new BuildStoreConsistencyState(), logger)
 		{
 			_settings = settings;
 			_leaderElection = leaderElection;
 			_buildStore = buildStore;
-			_blobStore = blobStore;
+			_blobService = blobService;
 			_blockStore = blockStore;
 			_refService = refService;
 			_referenceResolver = referenceResolver;
@@ -136,7 +136,7 @@ namespace Jupiter.Implementation
 									{
 										Interlocked.Increment(ref countOfBlobsChecked);
 
-										if (!await _blobStore.ExistsAsync(ns, blobId, cancellationToken: innerToken))
+										if (!await _blobService.ExistsAsync(ns, blobId, cancellationToken: innerToken))
 										{
 											_logger.LogWarning("Found missing blob {BlobId} in {Namespace}.{Bucket}.{BuildId}.{PartId}.", blobId, ns, bucket, buildId, partId);
 											Interlocked.Increment(ref countOfMissingBlobsFound);
@@ -148,7 +148,7 @@ namespace Jupiter.Implementation
 										{
 											Interlocked.Increment(ref countOfBlocksChecked);
 											
-											if (!await _blobStore.ExistsAsync(ns, metadataHash, cancellationToken: innerToken))
+											if (!await _blobService.ExistsAsync(ns, metadataHash, cancellationToken: innerToken))
 											{
 												_logger.LogWarning("Found missing block metadata {MetadataHash} for blob {BlobId} in {Namespace}.{Bucket}.{BuildId}.{PartId}.", metadataHash, blobId, ns, bucket, buildId, partId);
 												Interlocked.Increment(ref countOfMissingBlocksFound);
