Cloud DDC - Catch BlobNotFoundException when retrieving referenced blobs from BuildStoreConsistencyCheckService

#Synced-CL 48106205

diff --git a/Engine/Source/Programs/UnrealCloudDDC/Jupiter/Implementation/Builds/BuildStoreConsistencyCheckService.cs b/Engine/Source/Programs/UnrealCloudDDC/Jupiter/Implementation/Builds/BuildStoreConsistencyCheckService.cs
--- a/Engine/Source/Programs/UnrealCloudDDC/Jupiter/Implementation/Builds/BuildStoreConsistencyCheckService.cs
+++ b/Engine/Source/Programs/UnrealCloudDDC/Jupiter/Implementation/Builds/BuildStoreConsistencyCheckService.cs
@@ -115,48 +115,56 @@ namespace Jupiter.Implementation
 						await foreach ((string _, CbObjectId partId) in _buildStore.GetBuildPartsAsync(ns, bucket, buildId).WithCancellation(token))
 						{
 							RefId refId = RefId.FromName($"{buildId}/{partId}");
-							List<BlobId> referencedBlobs = await _refService.GetReferencedBlobsAsync(ns, bucket, refId, ignoreMissingBlobs: true, cancellationToken: token);
-
-							await Parallel.ForEachAsync(referencedBlobs, new ParallelOptions
+							
+							try
 							{
-								CancellationToken = token,
-								MaxDegreeOfParallelism = _settings.CurrentValue.BuildStoreConsistencyMaxParallelBlobs
-							},
-								async (referencedBlobId, innerToken) =>
-								{
-									BlobId[]? blobIds = await _referenceResolver.ResolveIdAsync(ns, referencedBlobId, cancellationToken: innerToken);
-									if (blobIds == null)
-									{
-										_logger.LogWarning("Unable to resolve blobId for referenced blob {ReferencedBlobId} in {Namespace}.{Bucket}.{BuildId}.{PartId}.", referencedBlobId, ns, bucket, buildId, partId);
-										Interlocked.Increment(ref countOfUnresolvableBlobs);
-										return;
-									}
+								List<BlobId> referencedBlobs = await _refService.GetReferencedBlobsAsync(ns, bucket, refId, ignoreMissingBlobs: true, cancellationToken: token);
 
-									foreach (BlobId blobId in blobIds)
+								await Parallel.ForEachAsync(referencedBlobs, new ParallelOptions
+								{
+									CancellationToken = token,
+									MaxDegreeOfParallelism = _settings.CurrentValue.BuildStoreConsistencyMaxParallelBlobs
+								},
+									async (referencedBlobId, innerToken) =>
 									{
-										Interlocked.Increment(ref countOfBlobsChecked);
-
-										if (!await _blobService.ExistsAsync(ns, blobId, cancellationToken: innerToken))
+										BlobId[]? blobIds = await _referenceResolver.ResolveIdAsync(ns, referencedBlobId, cancellationToken: innerToken);
+										if (blobIds == null)
 										{
-											_logger.LogWarning("Found missing blob {BlobId} in {Namespace}.{Bucket}.{BuildId}.{PartId}.", blobId, ns, bucket, buildId, partId);
-											Interlocked.Increment(ref countOfMissingBlobsFound);
+											_logger.LogWarning("Unable to resolve blobId for referenced blob {ReferencedBlobId} in {Namespace}.{Bucket}.{BuildId}.{PartId}.", referencedBlobId, ns, bucket, buildId, partId);
+											Interlocked.Increment(ref countOfUnresolvableBlobs);
+											return;
 										}
 
-										BlobId? metadataHash = await _blockStore.GetBlockMetadataAsync(ns, blobId);
-
-										if (metadataHash != null)
+										foreach (BlobId blobId in blobIds)
 										{
-											Interlocked.Increment(ref countOfBlocksChecked);
-											
-											if (!await _blobService.ExistsAsync(ns, metadataHash, cancellationToken: innerToken))
+											Interlocked.Increment(ref countOfBlobsChecked);
+
+											if (!await _blobService.ExistsAsync(ns, blobId, cancellationToken: innerToken))
 											{
-												_logger.LogWarning("Found missing block metadata {MetadataHash} for blob {BlobId} in {Namespace}.{Bucket}.{BuildId}.{PartId}.", metadataHash, blobId, ns, bucket, buildId, partId);
-												Interlocked.Increment(ref countOfMissingBlocksFound);
+												_logger.LogWarning("Found missing blob {BlobId} in {Namespace}.{Bucket}.{BuildId}.{PartId}.", blobId, ns, bucket, buildId, partId);
+												Interlocked.Increment(ref countOfMissingBlobsFound);
+											}
+
+											BlobId? metadataHash = await _blockStore.GetBlockMetadataAsync(ns, blobId);
+
+											if (metadataHash != null)
+											{
+												Interlocked.Increment(ref countOfBlocksChecked);
+												
+												if (!await _blobService.ExistsAsync(ns, metadataHash, cancellationToken: innerToken))
+												{
+													_logger.LogWarning("Found missing block metadata {MetadataHash} for blob {BlobId} in {Namespace}.{Bucket}.{BuildId}.{PartId}.", metadataHash, blobId, ns, bucket, buildId, partId);
+													Interlocked.Increment(ref countOfMissingBlocksFound);
+												}
 											}
 										}
 									}
-								}
-							);
+								);
+							}
+							catch (BlobNotFoundException e)
+							{
+								_logger.LogError(e, "Could not retrieve referenced blobs for {RefId} in {Namespace}.{Bucket}", refId, ns, bucket);
+							}
 						}
 					}
 				}
