Cloud DDC - Handle RefNotFoundExceptions in BuildStoreConsistencyCheckService

#Synced-CL 48106205

diff --git a/Engine/Source/Programs/UnrealCloudDDC/Jupiter/Implementation/Builds/BuildStoreConsistencyCheckService.cs b/Engine/Source/Programs/UnrealCloudDDC/Jupiter/Implementation/Builds/BuildStoreConsistencyCheckService.cs
--- a/Engine/Source/Programs/UnrealCloudDDC/Jupiter/Implementation/Builds/BuildStoreConsistencyCheckService.cs
+++ b/Engine/Source/Programs/UnrealCloudDDC/Jupiter/Implementation/Builds/BuildStoreConsistencyCheckService.cs
@@ -76,6 +76,8 @@ namespace Jupiter.Implementation
 			ulong countOfMissingBlobsFound = 0;
 			ulong countOfBlocksChecked = 0;
 			ulong countOfMissingBlocksFound = 0;
+			ulong countOfBlobErrors = 0;
+			ulong countOfRefErrors = 0;
 			await Parallel.ForEachAsync(_buildStore.ListAllBuildsAsync(cancellationToken), new ParallelOptions
 			{
 				CancellationToken = cancellationToken,
@@ -164,6 +166,12 @@ namespace Jupiter.Implementation
 							catch (BlobNotFoundException e)
 							{
 								_logger.LogError(e, "Could not retrieve referenced blobs for {RefId} in {Namespace}.{Bucket}", refId, ns, bucket);
+								Interlocked.Increment(ref countOfBlobErrors);
+							}
+							catch (RefNotFoundException e)
+							{
+								_logger.LogError(e, "Could not retrieve referenced blobs for {RefId} in {Namespace}.{Bucket}", refId, ns, bucket);
+								Interlocked.Increment(ref countOfRefErrors);
 							}
 						}
 					}
@@ -177,7 +185,7 @@ namespace Jupiter.Implementation
 			
 			if (_settings.CurrentValue.CheckBuildStoreMissingBlobs)
 			{
-				_logger.LogInformation("Build Store Consistency check finished. Found {CountOfUnresolvableBlobs} unresolvable blobs. Found {CountOfMissingBlobsFound} missing blobs ({CountOfBlobsChecked} total). Found {CountOfMissingBlocksFound} missing blocks ({CountOfBlocksChecked} total). Processed {CountOfBuilds} builds.", countOfUnresolvableBlobs, countOfMissingBlobsFound, countOfBlobsChecked, countOfMissingBlocksFound, countOfBlocksChecked, countOfBuildsChecked);
+				_logger.LogInformation("Build Store Consistency check finished. Found {CountOfUnresolvableBlobs} unresolvable blobs. Found {CountOfMissingBlobsFound} missing blobs ({CountOfBlobsChecked} total). Found {CountOfMissingBlocksFound} missing blocks ({CountOfBlocksChecked} total). Encountered {CountOfBlobErrors} BlobNotFoundExceptions. Encountered {CountOfRefErrors} RefNotFoundExceptions. Processed {CountOfBuilds} builds.", countOfUnresolvableBlobs, countOfMissingBlobsFound, countOfBlobsChecked, countOfMissingBlocksFound, countOfBlocksChecked, countOfBuildsChecked, countOfBlobErrors, countOfRefErrors);
 			}
 		}
 
