Cloud DDC - Add missed Scylla error handling paths

#Synced-CL 48106205

diff --git a/Engine/Source/Programs/UnrealCloudDDC/Jupiter/Controllers/BuildsController.cs b/Engine/Source/Programs/UnrealCloudDDC/Jupiter/Controllers/BuildsController.cs
--- a/Engine/Source/Programs/UnrealCloudDDC/Jupiter/Controllers/BuildsController.cs
+++ b/Engine/Source/Programs/UnrealCloudDDC/Jupiter/Controllers/BuildsController.cs
@@ -643,6 +643,19 @@ namespace Jupiter.Controllers
 				// cancellation is not an issue so we can suppress this
 				partialResult = true;
 			}
+			catch (AggregateException e)
+			{
+				if (e.InnerExceptions.Any(exception => exception.GetType() == typeof(RemoteServiceException)))
+				{
+					return Problem(e.Message, null, (int)HttpStatusCode.ServiceUnavailable);
+				}
+
+				throw;
+			}
+			catch (RemoteServiceException e)
+			{
+				return Problem(e.Message, null, (int)HttpStatusCode.ServiceUnavailable);
+			}
 			// Collect the results that was newest that match our limits and search
 			List<BuildMetadata> foundResultsInBuckets = foundResults.Values.Skip(searchRequest.Options.Skip).Take(searchRequest.Options.Limit).ToList();
 
@@ -2285,13 +2298,30 @@ namespace Jupiter.Controllers
 			responseObject.BeginArray("blocks");
 
 			BlockContext context = BlockContext.FromObject(buildRecord.BuildObject, useBaseBranch: false);
-			bool blocksFound = await WriteBlockMetadata(context, responseObject);
 
-			if (!blocksFound)
+			try
+			{
+				bool blocksFound = await WriteBlockMetadata(context, responseObject);
+
+				if (!blocksFound)
+				{
+					// no blocks found in the branch, fallback to the base branch
+					BlockContext baseContext = BlockContext.FromObject(buildRecord.BuildObject, useBaseBranch: true);
+					await WriteBlockMetadata(baseContext, responseObject);
+				}
+			}
+			catch (AggregateException e)
 			{
-				// no blocks found in the branch, fallback to the base branch
-				BlockContext baseContext = BlockContext.FromObject(buildRecord.BuildObject, useBaseBranch: true);
-				await WriteBlockMetadata(baseContext, responseObject);
+				if (e.InnerExceptions.Any(exception => exception.GetType() == typeof(RemoteServiceException)))
+				{
+					return Problem(e.Message, null, (int)HttpStatusCode.ServiceUnavailable);
+				}
+
+				throw;
+			}
+			catch (RemoteServiceException e)
+			{
+				return Problem(e.Message, null, (int)HttpStatusCode.ServiceUnavailable);
 			}
 
 			responseObject.EndArray();
diff --git a/Engine/Source/Programs/UnrealCloudDDC/Jupiter/Implementation/Blob/ScyllaBlobIndex.cs b/Engine/Source/Programs/UnrealCloudDDC/Jupiter/Implementation/Blob/ScyllaBlobIndex.cs
--- a/Engine/Source/Programs/UnrealCloudDDC/Jupiter/Implementation/Blob/ScyllaBlobIndex.cs
+++ b/Engine/Source/Programs/UnrealCloudDDC/Jupiter/Implementation/Blob/ScyllaBlobIndex.cs
@@ -433,7 +433,23 @@ public class ScyllaBlobIndex : IBlobIndex
 			refUpdateTasks[i] = _mapper.InsertAsync<ScyllaBlobIncomingReference>(incomingReference);
 		}
 
-		await Task.WhenAll(refUpdateTasks);
+		try
+		{
+			await Task.WhenAll(refUpdateTasks);
+		}
+		catch (AggregateException e)
+		{
+			if (e.InnerExceptions.Any(exception => exception.GetType() == typeof(QueryExecutionException)))
+			{
+				throw new RemoteServiceException("Error communicating with Scylla", e);
+			}
+
+			throw;
+		}
+		catch (QueryExecutionException e)
+		{
+			throw new RemoteServiceException("Error communicating with Scylla", e);
+		}
 	}
 
 	public async Task RemoveReferencesAsync(NamespaceId ns, BlobId id, List<BaseBlobReference>? referencesToRemove, CancellationToken cancellationToken)
diff --git a/Engine/Source/Programs/UnrealCloudDDC/Jupiter/Implementation/ContentId/ScyllaContentIdStore.cs b/Engine/Source/Programs/UnrealCloudDDC/Jupiter/Implementation/ContentId/ScyllaContentIdStore.cs
--- a/Engine/Source/Programs/UnrealCloudDDC/Jupiter/Implementation/ContentId/ScyllaContentIdStore.cs
+++ b/Engine/Source/Programs/UnrealCloudDDC/Jupiter/Implementation/ContentId/ScyllaContentIdStore.cs
@@ -179,14 +179,22 @@ namespace Jupiter.Implementation
 		public async Task PutAsync(NamespaceId ns, ContentId contentId, BlobId[] blobIdentifiers, int contentWeight, CancellationToken cancellationToken)
 		{
 			using TelemetrySpan scope = _tracer.BuildScyllaSpan("ScyllaContentIdStore.PutContentId").SetAttribute("resource.name", contentId.ToString());
-			if (_scyllaSessionManager.IsScylla)
+
+			try
 			{
-				await _mapper.UpdateAsync<ScyllaContentId>("SET chunks = ? WHERE content_id = ? AND content_weight = ?", blobIdentifiers.Select(id => new ScyllaBlobIdentifier(id)).ToArray(), new ScyllaBlobIdentifier(contentId), contentWeight);
+				if (_scyllaSessionManager.IsScylla)
+				{
+					await _mapper.UpdateAsync<ScyllaContentId>("SET chunks = ? WHERE content_id = ? AND content_weight = ?", blobIdentifiers.Select(id => new ScyllaBlobIdentifier(id)).ToArray(), new ScyllaBlobIdentifier(contentId), contentWeight);
 
+				}
+				else
+				{
+					await _mapper.UpdateAsync<CassandraContentId>("SET chunks = ? WHERE content_id = ? AND content_weight = ?", blobIdentifiers.Select(id => id.HashData).ToArray(), contentId.HashData, contentWeight);
+				}
 			}
-			else
+			catch (QueryExecutionException e)
 			{
-				await _mapper.UpdateAsync<CassandraContentId>("SET chunks = ? WHERE content_id = ? AND content_weight = ?", blobIdentifiers.Select(id => id.HashData).ToArray(), contentId.HashData, contentWeight);
+				throw new RemoteServiceException("Error communicating with Scylla", e);
 			}
 		}
 	}
diff --git a/Engine/Source/Programs/UnrealCloudDDC/Jupiter/Implementation/Objects/ScyllaReferencesStore.cs b/Engine/Source/Programs/UnrealCloudDDC/Jupiter/Implementation/Objects/ScyllaReferencesStore.cs
--- a/Engine/Source/Programs/UnrealCloudDDC/Jupiter/Implementation/Objects/ScyllaReferencesStore.cs
+++ b/Engine/Source/Programs/UnrealCloudDDC/Jupiter/Implementation/Objects/ScyllaReferencesStore.cs
@@ -537,7 +537,16 @@ namespace Jupiter.Implementation
 			using TelemetrySpan scope = _tracer.BuildScyllaSpan("scylla.get_namespaces");
 
 			{
-				RowSet rowSet = await _session.ExecuteAsync(_getNamespacesStatement.Bind());
+				RowSet rowSet;
+
+				try
+				{
+					rowSet = await _session.ExecuteAsync(_getNamespacesStatement.Bind());
+				}
+				catch (QueryExecutionException e)
+				{
+					throw new RemoteServiceException("Error communicating with Scylla", e);
+				}
 
 				foreach (Row row in rowSet)
 				{
@@ -555,7 +564,16 @@ namespace Jupiter.Implementation
 				// this will likely generate duplicates from the statements above but that is not a huge issue
 				using TelemetrySpan _ = _tracer.BuildScyllaSpan("scylla.get_old_namespaces");
 
-				RowSet rowSet = await _session.ExecuteAsync(_getNamespacesOldStatement.Bind());
+				RowSet rowSet;
+
+				try
+				{
+					rowSet = await _session.ExecuteAsync(_getNamespacesOldStatement.Bind());
+				}
+				catch (QueryExecutionException e)
+				{
+					throw new RemoteServiceException("Error communicating with Scylla", e);
+				}
 
 				foreach (Row row in rowSet)
 				{
