Cloud DDC - Skip build consistency checks for builds that are not finalized

#Synced-CL 48106205

diff --git a/Engine/Source/Programs/UnrealCloudDDC/Jupiter/Implementation/Builds/BuildStoreConsistencyCheckService.cs b/Engine/Source/Programs/UnrealCloudDDC/Jupiter/Implementation/Builds/BuildStoreConsistencyCheckService.cs
--- a/Engine/Source/Programs/UnrealCloudDDC/Jupiter/Implementation/Builds/BuildStoreConsistencyCheckService.cs
+++ b/Engine/Source/Programs/UnrealCloudDDC/Jupiter/Implementation/Builds/BuildStoreConsistencyCheckService.cs
@@ -70,6 +70,7 @@ namespace Jupiter.Implementation
 		private async Task RunConsistencyCheckAsync(CancellationToken cancellationToken)
 		{
 			ulong countOfBuildsChecked = 0;
+			ulong countOfUnfinalizedBuilds = 0;
 			ulong countOfIncorrectBuildsFound = 0;
 			ulong countOfUnresolvableBlobs = 0;
 			ulong countOfBlobsChecked = 0;
@@ -85,7 +86,7 @@ namespace Jupiter.Implementation
 			},
 				async (tuple, token) =>
 				{
-					(NamespaceId ns, BucketId bucket, CbObjectId buildId) = tuple;
+					(NamespaceId ns, BucketId bucket, CbObjectId buildId, bool isFinalized) = tuple;
 					Interlocked.Increment(ref countOfBuildsChecked);
 
 					if (countOfBuildsChecked % 100 == 0)
@@ -114,6 +115,13 @@ namespace Jupiter.Implementation
 
 					if (_settings.CurrentValue.CheckBuildStoreMissingBlobs)
 					{
+						if (!isFinalized)
+						{
+							_logger.LogInformation("Build {BuildId} in {Namespace}.{Bucket} is not finalized. Skipping.", buildId, ns, bucket);
+							Interlocked.Increment(ref countOfUnfinalizedBuilds);
+							return;
+						}
+
 						await foreach ((string _, CbObjectId partId) in _buildStore.GetBuildPartsAsync(ns, bucket, buildId).WithCancellation(token))
 						{
 							RefId refId = RefId.FromName($"{buildId}/{partId}");
@@ -185,7 +193,7 @@ namespace Jupiter.Implementation
 			
 			if (_settings.CurrentValue.CheckBuildStoreMissingBlobs)
 			{
-				_logger.LogInformation("Build Store Consistency check finished. Found {CountOfUnresolvableBlobs} unresolvable blobs. Found {CountOfMissingBlobsFound} missing blobs ({CountOfBlobsChecked} total). Found {CountOfMissingBlocksFound} missing blocks ({CountOfBlocksChecked} total). Encountered {CountOfBlobErrors} BlobNotFoundExceptions. Encountered {CountOfRefErrors} RefNotFoundExceptions. Processed {CountOfBuilds} builds.", countOfUnresolvableBlobs, countOfMissingBlobsFound, countOfBlobsChecked, countOfMissingBlocksFound, countOfBlocksChecked, countOfBuildsChecked, countOfBlobErrors, countOfRefErrors);
+				_logger.LogInformation("Build Store Consistency check finished. Found {CountOfUnresolvableBlobs} unresolvable blobs. Found {CountOfMissingBlobsFound} missing blobs ({CountOfBlobsChecked} total). Found {CountOfMissingBlocksFound} missing blocks ({CountOfBlocksChecked} total). Encountered {CountOfBlobErrors} BlobNotFoundExceptions. Encountered {CountOfRefErrors} RefNotFoundExceptions. Processed {CountOfBuilds} builds ({CountOfUnfinalizedBuilds} not finalized).", countOfUnresolvableBlobs, countOfMissingBlobsFound, countOfBlobsChecked, countOfMissingBlocksFound, countOfBlocksChecked, countOfBlobErrors, countOfRefErrors, countOfBuildsChecked, countOfUnfinalizedBuilds);
 			}
 		}
 
diff --git a/Engine/Source/Programs/UnrealCloudDDC/Jupiter/Implementation/Builds/IBuildStore.cs b/Engine/Source/Programs/UnrealCloudDDC/Jupiter/Implementation/Builds/IBuildStore.cs
--- a/Engine/Source/Programs/UnrealCloudDDC/Jupiter/Implementation/Builds/IBuildStore.cs
+++ b/Engine/Source/Programs/UnrealCloudDDC/Jupiter/Implementation/Builds/IBuildStore.cs
@@ -25,7 +25,7 @@ public interface IBuildStore
 	Task UpdateTTL(NamespaceId ns, BucketId bucket, CbObjectId buildId, uint ttl);
 	Task DeleteBuild(NamespaceId ns, BucketId bucket, CbObjectId buildId);
 	Task<uint?> GetTTL(NamespaceId ns, BucketId bucket, CbObjectId buildId);
-	IAsyncEnumerable<(NamespaceId, BucketId, CbObjectId)> ListAllBuildsAsync(CancellationToken cancellationToken);
+	IAsyncEnumerable<(NamespaceId, BucketId, CbObjectId, bool)> ListAllBuildsAsync(CancellationToken cancellationToken);
 }
 
 public class BuildMetadata
diff --git a/Engine/Source/Programs/UnrealCloudDDC/Jupiter/Implementation/Builds/MemoryBuildStore.cs b/Engine/Source/Programs/UnrealCloudDDC/Jupiter/Implementation/Builds/MemoryBuildStore.cs
--- a/Engine/Source/Programs/UnrealCloudDDC/Jupiter/Implementation/Builds/MemoryBuildStore.cs
+++ b/Engine/Source/Programs/UnrealCloudDDC/Jupiter/Implementation/Builds/MemoryBuildStore.cs
@@ -161,12 +161,12 @@ public class MemoryBuildStore : IBuildStore
 		return Task.FromResult(_builds.TryGetValue(key, out MemoryBuildRecord? value) ? (uint?)(value.ExpireTime - DateTime.Now).TotalSeconds : null);
 	}
 
-	public async IAsyncEnumerable<(NamespaceId, BucketId, CbObjectId)> ListAllBuildsAsync([EnumeratorCancellation] CancellationToken cancellationToken)
+	public async IAsyncEnumerable<(NamespaceId, BucketId, CbObjectId, bool)> ListAllBuildsAsync([EnumeratorCancellation] CancellationToken cancellationToken)
 	{
 		await Task.CompletedTask;
 		foreach (MemoryBuildRecord build in _builds.Values.OrderByDescending(record => record.CreationDate))
 		{
-			yield return (build.Namespace, build.Bucket, build.BuildId);
+			yield return (build.Namespace, build.Bucket, build.BuildId, build.IsFinalized);
 		}
 	}
 
diff --git a/Engine/Source/Programs/UnrealCloudDDC/Jupiter/Implementation/Builds/ScyllaBuildStore.cs b/Engine/Source/Programs/UnrealCloudDDC/Jupiter/Implementation/Builds/ScyllaBuildStore.cs
--- a/Engine/Source/Programs/UnrealCloudDDC/Jupiter/Implementation/Builds/ScyllaBuildStore.cs
+++ b/Engine/Source/Programs/UnrealCloudDDC/Jupiter/Implementation/Builds/ScyllaBuildStore.cs
@@ -253,7 +253,7 @@ public class ScyllaBuildStore : IBuildStore
 		return null;
 	}
 
-	public async IAsyncEnumerable<(NamespaceId, BucketId, CbObjectId)> ListAllBuildsAsync([EnumeratorCancellation] CancellationToken cancellationToken)
+	public async IAsyncEnumerable<(NamespaceId, BucketId, CbObjectId, bool)> ListAllBuildsAsync([EnumeratorCancellation] CancellationToken cancellationToken)
 	{
 		using TelemetrySpan scope = _tracer.BuildScyllaSpan("scylla.get_all_builds");
 
@@ -271,7 +271,7 @@ public class ScyllaBuildStore : IBuildStore
 		
 		foreach (ScyllaBuildEntry buildEntry in entries)
 		{
-			yield return (new NamespaceId(buildEntry.Namespace), new BucketId(buildEntry.BucketId), new CbObjectId(buildEntry.BuildId));
+			yield return (new NamespaceId(buildEntry.Namespace), new BucketId(buildEntry.BucketId), new CbObjectId(buildEntry.BuildId), buildEntry.IsFinalized);
 		}
 	}
 
