Cloud Storage - Added suffix to server timing for replication to indicate which region we replicated from

#Synced-CL 49042760

diff --git a/Engine/Source/Programs/UnrealCloudDDC/Jupiter/Implementation/Blob/BlobService.cs b/Engine/Source/Programs/UnrealCloudDDC/Jupiter/Implementation/Blob/BlobService.cs
--- a/Engine/Source/Programs/UnrealCloudDDC/Jupiter/Implementation/Blob/BlobService.cs
+++ b/Engine/Source/Programs/UnrealCloudDDC/Jupiter/Implementation/Blob/BlobService.cs
@@ -636,8 +636,6 @@ public class BlobService : IBlobService
 		IServerTiming? serverTiming = _httpContextAccessor.HttpContext?.RequestServices.GetService<IServerTiming>();
 		using TelemetrySpan scope = _tracer.StartActiveSpan("HierarchicalStore.Replicate").SetAttribute("operation.name", "HierarchicalStore.Replicate").SetAttribute("resource.name", blob.ToString());
 
-		using ServerTimingMetricScoped? serverTimingScope = serverTiming?.CreateServerTimingMetricScope("blob.replicate", "Replicating blob from remote instances");
-
 		List<string> regions;
 		try
 		{
@@ -669,6 +667,8 @@ public class BlobService : IBlobService
 
 			_logger.LogInformation("Attempting to replicate blob {Blob} in Namespace {Namespace} from {Region}", blob, ns, region);
 
+			using ServerTimingMetricScoped? serverTimingScope = serverTiming?.CreateServerTimingMetricScope($"blob.replicate.{region}", "Replicating blob from remote instances");
+
 			PeerEndpoints peerEndpoint = peerStatus.Endpoints.First();
 			using HttpClient httpClient = _httpClientFactory.CreateClient();
 			string url = peerEndpoint.Url.ToString();
