Cloud DDC - Add ability to customize nginx proxy_read_timeout

#Synced-CL 48106205

diff --git a/Engine/Source/Programs/UnrealCloudDDC/Helm/UnrealCloudDDC/templates/configmap-nginx.yaml b/Engine/Source/Programs/UnrealCloudDDC/Helm/UnrealCloudDDC/templates/configmap-nginx.yaml
--- a/Engine/Source/Programs/UnrealCloudDDC/Helm/UnrealCloudDDC/templates/configmap-nginx.yaml
+++ b/Engine/Source/Programs/UnrealCloudDDC/Helm/UnrealCloudDDC/templates/configmap-nginx.yaml
@@ -45,6 +45,9 @@ data:
         proxy_http_version 1.1;
         
         location /api {
+          {{ if .Values.nginx.readTimeout }}
+          proxy_read_timeout {{ .Values.nginx.readTimeout }};
+          {{- end }}
           proxy_set_header X-Jupiter-XAccel-Supported true;
           proxy_set_header X-Jupiter-Port Public;
           proxy_set_header Connection "";
@@ -97,6 +100,9 @@ data:
         proxy_http_version 1.1;
         
         location /api {
+          {{ if .Values.nginx.readTimeout }}
+          proxy_read_timeout {{ .Values.nginx.readTimeout }};
+          {{- end }}
           proxy_set_header X-Jupiter-XAccel-Supported true;
           proxy_set_header X-Jupiter-Port Corp;
           proxy_set_header Connection "";
@@ -150,6 +156,9 @@ data:
         proxy_http_version 1.1;
         
         location /api {
+          {{ if .Values.nginx.readTimeout }}
+          proxy_read_timeout {{ .Values.nginx.readTimeout }};
+          {{- end }}
           proxy_set_header X-Jupiter-XAccel-Supported true;
           proxy_set_header X-Jupiter-Port Public;
           proxy_set_header Connection "";
@@ -203,6 +212,9 @@ data:
         proxy_http_version 1.1;
         
         location /api {
+          {{ if .Values.nginx.readTimeout }}
+          proxy_read_timeout {{ .Values.nginx.readTimeout }};
+          {{- end }}
           proxy_set_header X-Jupiter-XAccel-Supported true;
           proxy_set_header X-Jupiter-Port Corp;
           proxy_set_header Connection "";
diff --git a/Engine/Source/Programs/UnrealCloudDDC/Helm/UnrealCloudDDC/values.yaml b/Engine/Source/Programs/UnrealCloudDDC/Helm/UnrealCloudDDC/values.yaml
--- a/Engine/Source/Programs/UnrealCloudDDC/Helm/UnrealCloudDDC/values.yaml
+++ b/Engine/Source/Programs/UnrealCloudDDC/Helm/UnrealCloudDDC/values.yaml
@@ -378,6 +378,9 @@ nginx:
 
   ## The length of the nginx backlog, bump this if you increase `net.core.somaxconn` beyond 4096
   backlog: 4096
+  
+  ## Proxy read timeout. Increase this if experiencing 504 errors during large operations.
+  readTimeout: 60s
 
   ## Liveness and readiness probes, defaults to corresponding endpoints
   ##
