Cloud Storage - Added ReadPartObject access control action which is required to be able to read a part of a build. Also added ability to control actions assigned to certain resources and implemented parts for that, which allows us to control access to some parts but not others. Will be used to have more sensitive parts with data that we do not want everyone to have access to.

#Synced-CL 49042760

diff --git a/Engine/Source/Programs/UnrealCloudDDC/Jupiter.Common/Authentication/Acl.cs b/Engine/Source/Programs/UnrealCloudDDC/Jupiter.Common/Authentication/Acl.cs
--- a/Engine/Source/Programs/UnrealCloudDDC/Jupiter.Common/Authentication/Acl.cs
+++ b/Engine/Source/Programs/UnrealCloudDDC/Jupiter.Common/Authentication/Acl.cs
@@ -2,6 +2,7 @@
 
 using System;
 using System.Collections.Generic;
+using System.ComponentModel.DataAnnotations;
 using System.Security.Claims;
 using System.Text.RegularExpressions;
 using Microsoft.AspNetCore.Authorization;
@@ -43,6 +44,16 @@ public enum JupiterAclAction
 	/// </summary>
 	WriteTransactionLog,
 
+	/// <summary>
+	/// Access to read build parts
+	/// </summary>
+	ReadBuildPart,
+
+	/// <summary>
+	/// Access to write build parts
+	/// </summary>
+	WriteBuildPart,
+
 	/// <summary>
 	/// Access to perform administrative task
 	/// </summary>
@@ -236,15 +247,40 @@ public class AclPolicy
 			}
 		}
 
-		List<JupiterAclAction> allowedActions = new List<JupiterAclAction>();
+		bool requiresResourceScope = scope.ResourceType != null && scope.ResourceName != null;
 
-		if (bucketFound)
+		if (!requiresResourceScope && bucketFound)
 		{
 			// this policy match all claims and scopes, we grant the actions it contains
-			allowedActions.AddRange(Actions);
+			return Actions;
 		}
 
-		return allowedActions;
+		if (requiresResourceScope)
+		{
+			bool resourceFound = false;
+			foreach (AclResourceScopeEntry resourceScope in Scope.Resources)
+			{
+				if (string.Equals(resourceScope.ResourceType, scope.ResourceType, StringComparison.OrdinalIgnoreCase))
+				{
+					bool isMatch = resourceScope.Matches(scope.ResourceName ?? string.Empty, logger, Name);
+
+					if (isMatch)
+					{
+						logger.LogDebug("Resource scope matched for policy \'{Name}\'.", Name);
+						resourceFound = true;
+						break;
+					}
+				}
+			}
+
+			if (resourceFound)
+			{
+				return Actions;
+			}
+		}
+
+		// policy did not match
+		return Array.Empty<JupiterAclAction>();
 	}
 }
 
@@ -267,6 +303,22 @@ public class AclScope
 	[System.Diagnostics.CodeAnalysis.SuppressMessage("Usage", "CA2227:Collection properties should be read only", Justification = "Used by the configuration system")]
 	// ReSharper disable once CollectionNeverUpdated.Global
 	public List<AclScopeEntry> Buckets { get; set; } = new List<AclScopeEntry>();
+
+	/// <summary>
+	///  List of one or more resources (by-name and type) that this policy applies to
+	/// </summary>
+	[System.Diagnostics.CodeAnalysis.SuppressMessage("Usage", "CA2227:Collection properties should be read only", Justification = "Used by the configuration system")]
+	// ReSharper disable once CollectionNeverUpdated.Global
+	public List<AclResourceScopeEntry> Resources { get; set; } = new List<AclResourceScopeEntry>();
+}
+
+/// <summary>
+/// Potential mapping of a scope for a set of resources
+/// </summary>
+public class AclResourceScopeEntry : AclScopeEntry
+{
+	[Required]
+	public string? ResourceType { get; set; } = null;
 }
 
 /// <summary>
diff --git a/Engine/Source/Programs/UnrealCloudDDC/Jupiter.Common/Authentication/ScopeAuthorizationHandler.cs b/Engine/Source/Programs/UnrealCloudDDC/Jupiter.Common/Authentication/ScopeAuthorizationHandler.cs
--- a/Engine/Source/Programs/UnrealCloudDDC/Jupiter.Common/Authentication/ScopeAuthorizationHandler.cs
+++ b/Engine/Source/Programs/UnrealCloudDDC/Jupiter.Common/Authentication/ScopeAuthorizationHandler.cs
@@ -25,9 +25,20 @@ namespace Jupiter
 			Bucket = bucket;
 		}
 
+		public AccessScope(NamespaceId? ns, BucketId? bucket, string resourceType, string resourceName)
+		{
+			Namespace = ns;
+			Bucket = bucket;
+			ResourceType = resourceType;
+			ResourceName = resourceName;
+		}
+
 		public static AccessScope GlobalScope { get; } = new AccessScope(null, null);
 		public NamespaceId? Namespace { get; }
 		public BucketId? Bucket { get; }
+
+		public string? ResourceType { get; } = null;
+		public string? ResourceName { get; } = null;
 	}
 
 	public class ScopeAuthorizationHandler : AuthorizationHandler<ScopeAccessRequirement, ScopeAccessRequest>
diff --git a/Engine/Source/Programs/UnrealCloudDDC/Jupiter/Controllers/BuildsController.cs b/Engine/Source/Programs/UnrealCloudDDC/Jupiter/Controllers/BuildsController.cs
--- a/Engine/Source/Programs/UnrealCloudDDC/Jupiter/Controllers/BuildsController.cs
+++ b/Engine/Source/Programs/UnrealCloudDDC/Jupiter/Controllers/BuildsController.cs
@@ -70,6 +70,11 @@ namespace Jupiter.Controllers
 		private readonly IReferenceResolver _referenceResolver;
 		private readonly ILogger? _auditLogger;
 
+		/// <summary>
+		/// Constant used for access control of parts
+		/// </summary>
+		private const string BuildPartType = "parts";
+
 		public BuildsController(IRefService refService, IBlobService blobStore, IBlockStore blockStore, IBuildStore buildStore, IBlobIndex blobIndex, IReferenceResolver referenceResolver, IDiagnosticContext diagnosticContext, FormatResolver formatResolver, BufferedPayloadFactory bufferedPayloadFactory, NginxRedirectHelper nginxRedirectHelper, INamespacePolicyResolver namespacePolicyResolver, IRequestHelper requestHelper, Tracer tracer, ILogger<BuildsController> logger, IOptionsMonitor<BuildsSettings> settings)
 		{
 			_refService = refService;
@@ -195,8 +200,14 @@ namespace Jupiter.Controllers
 
 			try
 			{
-				await foreach ((string _, CbObjectId partId) in _buildStore.GetBuildPartsAsync(ns, bucket, buildId))
+				await foreach ((string partName, CbObjectId partId) in _buildStore.GetBuildPartsAsync(ns, bucket, buildId))
 				{
+					ActionResult? accessResultPart = await _requestHelper.HasAccessToScopeAsync(User, Request, new AccessScope(ns, bucket, BuildPartType, partName), new[] { JupiterAclAction.WriteBuildPart });
+					if (accessResultPart != null)
+					{
+						return accessResultPart;
+					}
+
 					RefId refId = RefId.FromName($"{buildId}/{partId}");
 					try
 					{
@@ -270,6 +281,8 @@ namespace Jupiter.Controllers
 			{
 				await foreach ((string _, CbObjectId partId) in _buildStore.GetBuildPartsAsync(ns, bucket, buildId))
 				{
+					// we do not access control updating the individual part here as it would cause inconsistencies if some parts were updated but others weren't
+
 					RefId refId = RefId.FromName($"{buildId}/{partId}");
 					await _refService.UpdateTTL(ns, bucket, refId, request.TTL);
 				}
@@ -303,6 +316,8 @@ namespace Jupiter.Controllers
 			{
 				await foreach ((string _, CbObjectId partId) in _buildStore.GetBuildPartsAsync(ns, bucket, buildId))
 				{
+					// we do not explicitly access control the deletion of the build part, makes no sense as you can not delete just individual parts instead you delete the whole build
+
 					RefId refId = RefId.FromName($"{buildId}/{partId}");
 					await _refService.DeleteAsync(ns, bucket, refId);
 				}
@@ -360,8 +375,12 @@ namespace Jupiter.Controllers
 				newObjectWriter.BeginUniformObject("parts", CbFieldType.ObjectId | CbFieldType.HasFieldName);
 				await foreach ((string partName, CbObjectId partId) in _buildStore.GetBuildPartsAsync(ns, bucket, buildId))
 				{
-
-					newObjectWriter.WriteObjectId(partName, partId);
+					ActionResult? accessResultPart = await _requestHelper.HasAccessToScopeAsync(User, Request, new AccessScope(ns, bucket, BuildPartType, partName), new[] { JupiterAclAction.ReadBuildPart });
+					if (accessResultPart == null)
+					{
+						// only include the build part if this request has access to it
+						newObjectWriter.WriteObjectId(partName, partId);
+					}
 				}
 				newObjectWriter.EndObject(); // parts
 				newObjectWriter.EndObject(); // root
@@ -746,7 +765,9 @@ namespace Jupiter.Controllers
 				BlockContext newBlockContext = BlockContext.FromObject(newBuildObject);
 
 				async Task CopyRefAsync(NamespaceId ns, BucketId bucket, RefId oldRefId, NamespaceId targetNamespace, BucketId targetBucket, RefId newRefId, CbObjectId? buildPart = null)
-				{ 
+				{
+					// we do not explicitly access control the creation of the build part, this part already existed so we should create a full copy of it
+
 					(RefRecord refRecord, BlobContents? contents) = await _refService.GetAsync(ns, bucket, oldRefId);
 					if (contents == null)
 					{
@@ -870,7 +891,7 @@ namespace Jupiter.Controllers
 		[ProducesDefaultResponseType]
 		public async Task<IActionResult> GetBuildPartAsync(NamespaceId ns, BucketId bucket, CbObjectId buildId, string partId, [FromRoute] string? format = null)
 		{
-			ActionResult? accessResult = await _requestHelper.HasAccessToScopeAsync(User, Request, new AccessScope(ns, bucket), new[] { JupiterAclAction.ReadObject });
+			ActionResult? accessResult = await _requestHelper.HasAccessToScopeAsync(User, Request, new AccessScope(ns, bucket, BuildPartType, partId), new[] { JupiterAclAction.ReadBuildPart });
 			if (accessResult != null)
 			{
 				return accessResult;
@@ -985,7 +1006,7 @@ namespace Jupiter.Controllers
 		[ProducesDefaultResponseType]
 		public async Task<IActionResult> PutBuildPartAsync(NamespaceId ns, BucketId bucket, CbObjectId buildId, CbObjectId partId, string partName)
 		{
-			ActionResult? accessResult = await _requestHelper.HasAccessToScopeAsync(User, Request, new AccessScope(ns, bucket), new[] { JupiterAclAction.WriteObject });
+			ActionResult? accessResult = await _requestHelper.HasAccessToScopeAsync(User, Request, new AccessScope(ns, bucket), new[] { JupiterAclAction.WriteBuildPart });
 			if (accessResult != null)
 			{
 				return accessResult;
@@ -1637,13 +1658,19 @@ namespace Jupiter.Controllers
 
 			try
 			{
-				await foreach ((string _, CbObjectId foundPartId) in _buildStore.GetBuildPartsAsync(ns, bucket, buildId))
+				await foreach ((string partName, CbObjectId foundPartId) in _buildStore.GetBuildPartsAsync(ns, bucket, buildId))
 				{
 					if (!CbObjectId.Equals(foundPartId, partId))
 					{
 						continue;
 					}
 
+					ActionResult? accessResultPart = await _requestHelper.HasAccessToScopeAsync(User, Request, new AccessScope(ns, bucket, BuildPartType, partName), new[] { JupiterAclAction.ReadBuildPart });
+					if (accessResultPart != null)
+					{
+						return accessResultPart;
+					}
+
 					(RefRecord refRecord, BlobContents? contents)  = await _refService.GetAsync(ns, bucket, RefId.FromName($"{buildId}/{partId}"));
 					if (contents == null)
 					{
@@ -1710,7 +1737,6 @@ namespace Jupiter.Controllers
 #pragma warning disable IDE0060
 			CbObjectId buildId, 
 			CbObjectId partId, 
-#pragma warning restore IDE0060
 #pragma warning restore IDE0060
 			[Required] [FromBody] BuildPartStatsRequest statsRequest)
 		{
diff --git a/Engine/Source/Programs/UnrealCloudDDC/Jupiter/appsettings.Testing.json b/Engine/Source/Programs/UnrealCloudDDC/Jupiter/appsettings.Testing.json
--- a/Engine/Source/Programs/UnrealCloudDDC/Jupiter/appsettings.Testing.json
+++ b/Engine/Source/Programs/UnrealCloudDDC/Jupiter/appsettings.Testing.json
@@ -66,6 +66,8 @@
 					"WriteObject",
 					"EnumerateBucket",
 					"DeleteObject",
+					"ReadBuildPart",
+					"WriteBuildPart",
 					"DeleteBucket",
 					"ReadTransactionLog",
 					"WriteTransactionLog",
@@ -84,6 +86,12 @@
 							"match": ".*\\..*\\..*\\.(.*)",
 							"not": true
 						}
+					],
+					"Resources": [
+						{
+							"values": [ "*" ],
+							"resourceType":  "parts"
+						}
 					]
 				}
 			},
@@ -92,6 +100,8 @@
 				"Actions": [
 					"ReadObject",
 					"WriteObject",
+					"ReadBuildPart",
+					"WriteBuildPart",
 					"EnumerateBucket"
 				],
 				"Scope": {
@@ -104,6 +114,12 @@
 						{
 							"values": [ "*" ]
 						}
+					],
+					"Resources": [
+						{
+							"values": [ "*" ],
+							"resourceType": "parts"
+						}
 					]
 				}
 			},
diff --git a/Engine/Source/Programs/UnrealCloudDDC/README.md b/Engine/Source/Programs/UnrealCloudDDC/README.md
--- a/Engine/Source/Programs/UnrealCloudDDC/README.md
+++ b/Engine/Source/Programs/UnrealCloudDDC/README.md
@@ -159,6 +159,8 @@ Access to operations within UnrealCloudDDC is controlled using a set of actions:
     - ReadObject
     - WriteObject
     - DeleteObject
+    - ReadBuildPart
+    - WriteBuildPart
     - EnumerateBucket
     - DeleteBucket
     - DeleteNamespace
@@ -189,6 +191,8 @@ Example configuration that sets transaction log access for users that have acces
       - ReadObject
       - WriteObject
       - EnumerateObject
+      - ReadBuildPart
+      - WriteBuildPart
       scope:
         namespaces:
         - values: 
@@ -209,20 +213,23 @@ Example configuration that sets transaction log access for users that have acces
       - ReadObject
       - WriteObject
       - DeleteObject
+      - ReadBuildPart
+      - WriteBuildPart
       - DeleteBucket
       - DeleteNamespace
       - AdminAction
 
 ```
 If you specify multiple claims in the claims array these are ANDed together thus requiring all the claims. A claim statement like `A=B` requires claim `A` to have value `B` (or contain value `B` in the case of a array).
-You can also specify the `*` claim which grants access for any valid token no matter which claims it has, this is mostly for debug / testing scenarios and shouldnt be used for production data.
+You can also specify the `*` claim which grants access for any valid token no matter which claims it has, this is mostly for debug / testing scenarios and should not be used for production data.
 
 If you specify multiple values under namespace and/or bucket they are ORed (it matches that list of namespaces and buckets).
-Multiple statments under the namespaces and buckets directly are ANDed (all of them have to be satisfied).
+Multiple statements under the namespaces and buckets directly are ANDed (all of them have to be satisfied).
 
 You can also use the `match` and `not` values to create more advanced rules that exclude certain patterns and uses regular expressions to extract only a part of the bucket or namespace.
 
-This is an example of an advanced setup for builds that exclude access to certain platforms (in this example linux is considered a secret platform):
+This is an example of an advanced setup for builds that exclude access to certain platforms (in this example linux is considered a secret platform) as well as excluding access to a build part called "sensitive" unless 
+explicit access is granted via a seperate claim:
 
 ```
     - claims:
@@ -237,10 +244,18 @@ This is an example of an advanced setup for builds that exclude access to certai
           - "linux"
           match: .*\..*\.*\.(.*)
           not: true
+        resources:
+        # access to all parts except the sensitive part
+        - resourceType: "parts"
+          values: 
+          - "sensitive"
+          not: true
       actions:
       - ReadObject
       - WriteObject
       - EnumerateObject
+      - ReadBuildPart
+      - WriteBuildPart
 
     - claims:
       - groups=app-ddc-storage-ue-platform-linux
@@ -256,6 +271,23 @@ This is an example of an advanced setup for builds that exclude access to certai
         - values: 
           - linux
           match: .*\..*\.*\.(.*)
+
+    - claims:
+      - groups=app-ddc-storage-ue-sensitive-access
+      actions:
+      - ReadBuildPart
+      - WriteBuildPart
+      scopes:
+        namespaces:
+        - values: 
+          - ue.builds
+        buckets:
+        - values: 
+          - *
+        resources:
+        - resourceType: "parts"
+          values: 
+          - "sensitive"
 ```
 
 # Common Configuration
@@ -279,20 +311,7 @@ This is an example of how the policies for a namespace could be configured for o
 ```
 namespace:
   policies:
-    ue.ddc:
-      acls:
-      - actions: 
-        - ReadObject
-        - WriteObject
-        claims: 
-        - ExampleClaim   
     ue.oplog:
-      acls:
-      - actions: 
-        - ReadObject
-        - WriteObject
-        claims: 
-        - ExampleClaim   
       storagePool: "oplog"
       UseBlobIndexForExists: true
       GcMethod: TTL
diff --git a/Engine/Source/Programs/UnrealCloudDDC/changelog.md b/Engine/Source/Programs/UnrealCloudDDC/changelog.md
--- a/Engine/Source/Programs/UnrealCloudDDC/changelog.md
+++ b/Engine/Source/Programs/UnrealCloudDDC/changelog.md
@@ -1,5 +1,7 @@
 # Unreleased
 
+* Added `ReadBuildPart` and `WriteBuildPart` actions - these allows access control per part based on its name for builds. This allows for more sensitive parts of builds to still be distributed as the same thing but those files can be excluded from some users. This is a very advanced use case and we generally recommend to just assign Read and WriteBuildPart to everyone like you would `ReadObject` and `WriteObject` .
+
 # 1.4.0
 Noteworthy and potentially breaking changes:
 * All endpoints have been updated to use the new per-bucket authorization by default. Set the `Auth.UseLegacyConfiguration` setting to `true` to enable the legacy authentication behavior.
