Cloud DDC - Tracking the exception that occured when a 503 is reported in the tracing as while these are often retryable we still need to know what the underlying issue was

#Synced-CL 49272873

diff --git a/Engine/Source/Programs/UnrealCloudDDC/Jupiter/Controllers/ReferencesController.cs b/Engine/Source/Programs/UnrealCloudDDC/Jupiter/Controllers/ReferencesController.cs
--- a/Engine/Source/Programs/UnrealCloudDDC/Jupiter/Controllers/ReferencesController.cs
+++ b/Engine/Source/Programs/UnrealCloudDDC/Jupiter/Controllers/ReferencesController.cs
@@ -918,6 +918,7 @@ namespace Jupiter.Controllers
 			}
 			catch (RemoteServiceException e)
 			{
+				Tracer.CurrentSpan.RecordException(e);
 				return Problem(e.Message, null, (int)HttpStatusCode.ServiceUnavailable);
 			}
 
@@ -971,6 +972,7 @@ namespace Jupiter.Controllers
 			}
 			catch (RemoteServiceException e)
 			{
+				Tracer.CurrentSpan.RecordException(e);
 				return Problem(e.Message, null, (int)HttpStatusCode.ServiceUnavailable);
 			}
 		}
