Cloud DDC - Return explicit type identifier during a hash mismatch bad request to make it easier to detect this case

#Synced-CL 49272873

diff --git a/Engine/Source/Programs/UnrealCloudDDC/Jupiter/Controllers/BuildsController.cs b/Engine/Source/Programs/UnrealCloudDDC/Jupiter/Controllers/BuildsController.cs
--- a/Engine/Source/Programs/UnrealCloudDDC/Jupiter/Controllers/BuildsController.cs
+++ b/Engine/Source/Programs/UnrealCloudDDC/Jupiter/Controllers/BuildsController.cs
@@ -1055,7 +1055,8 @@ namespace Jupiter.Controllers
 			{
 				return BadRequest(new ProblemDetails
 				{
-					Title = $"Incorrect hash, got hash \"{e.SuppliedHash}\" but hash of content was determined to be \"{e.ContentHash}\""
+					Title = $"Incorrect hash, got hash \"{e.SuppliedHash}\" but hash of content was determined to be \"{e.ContentHash}\"",
+					Type = ProblemTypes.HashMismatchError
 				});
 			}
 			catch (ClientSendSlowException e)
@@ -1522,7 +1523,8 @@ namespace Jupiter.Controllers
 			{
 				return BadRequest(new ProblemDetails
 				{
-					Title = $"Incorrect hash, got hash \"{e.SuppliedHash}\" but hash of content was determined to be \"{e.ContentHash}\""
+					Title = $"Incorrect hash, got hash \"{e.SuppliedHash}\" but hash of content was determined to be \"{e.ContentHash}\"",
+					Type = ProblemTypes.HashMismatchError
 				});
 			}
 			catch (ClientSendSlowException e)
@@ -1887,8 +1889,8 @@ namespace Jupiter.Controllers
 			{
 				return BadRequest(new ProblemDetails
 				{
-					Title =
-						$"Incorrect hash, got hash \"{e.SuppliedHash}\" but hash of content was determined to be \"{e.ContentHash}\""
+					Title = $"Incorrect hash, got hash \"{e.SuppliedHash}\" but hash of content was determined to be \"{e.ContentHash}\"",
+					Type = ProblemTypes.HashMismatchError
 				});
 			}
 			catch (ResourceHasToManyRequestsException)
@@ -2022,8 +2024,8 @@ namespace Jupiter.Controllers
 			{
 				return BadRequest(new ProblemDetails
 				{
-					Title =
-						$"Incorrect hash, got hash \"{e.SuppliedHash}\" but hash of content was determined to be \"{e.ContentHash}\""
+					Title = $"Incorrect hash, got hash \"{e.SuppliedHash}\" but hash of content was determined to be \"{e.ContentHash}\"",
+					Type = ProblemTypes.HashMismatchError
 				});
 			}
 			
diff --git a/Engine/Source/Programs/UnrealCloudDDC/Jupiter/Controllers/CompressedBlobController.cs b/Engine/Source/Programs/UnrealCloudDDC/Jupiter/Controllers/CompressedBlobController.cs
--- a/Engine/Source/Programs/UnrealCloudDDC/Jupiter/Controllers/CompressedBlobController.cs
+++ b/Engine/Source/Programs/UnrealCloudDDC/Jupiter/Controllers/CompressedBlobController.cs
@@ -257,8 +257,8 @@ namespace Jupiter.Controllers
 			{
 				return BadRequest(new ProblemDetails
 				{
-					Title =
-						$"Incorrect hash, got hash \"{e.SuppliedHash}\" but hash of content was determined to be \"{e.ContentHash}\""
+					Title = $"Incorrect hash, got hash \"{e.SuppliedHash}\" but hash of content was determined to be \"{e.ContentHash}\"",
+					Type = ProblemTypes.HashMismatchError
 				});
 			}
 			catch (ClientSendSlowException e)
@@ -301,7 +301,8 @@ namespace Jupiter.Controllers
 			{
 				return BadRequest(new ProblemDetails
 				{
-					Title = $"Incorrect hash, got hash \"{e.SuppliedHash}\" but hash of content was determined to be \"{e.ContentHash}\""
+					Title = $"Incorrect hash, got hash \"{e.SuppliedHash}\" but hash of content was determined to be \"{e.ContentHash}\"",
+					Type = ProblemTypes.HashMismatchError
 				});
 			}
 			catch (ClientSendSlowException e)
diff --git a/Engine/Source/Programs/UnrealCloudDDC/Jupiter/Controllers/ProblemTypes.cs b/Engine/Source/Programs/UnrealCloudDDC/Jupiter/Controllers/ProblemTypes.cs
--- a/Engine/Source/Programs/UnrealCloudDDC/Jupiter/Controllers/ProblemTypes.cs
+++ b/Engine/Source/Programs/UnrealCloudDDC/Jupiter/Controllers/ProblemTypes.cs
@@ -6,5 +6,6 @@ namespace Jupiter.Controllers
 	{
 		public static string UseSnapshot { get; set; } = "http://jupiter.epicgames.com/replication/useSnapshot";
 		public static string NoDataFound { get; set; } = "http://jupiter.epicgames.com/replication/noReplicationData";
+		public static string HashMismatchError { get; set; } = "http://jupiter.epicgames.com/blob/hashMismatch";
 	}
 }
diff --git a/Engine/Source/Programs/UnrealCloudDDC/Jupiter/Controllers/ReferencesController.cs b/Engine/Source/Programs/UnrealCloudDDC/Jupiter/Controllers/ReferencesController.cs
--- a/Engine/Source/Programs/UnrealCloudDDC/Jupiter/Controllers/ReferencesController.cs
+++ b/Engine/Source/Programs/UnrealCloudDDC/Jupiter/Controllers/ReferencesController.cs
@@ -909,7 +909,8 @@ namespace Jupiter.Controllers
 			{
 				return BadRequest(new ProblemDetails
 				{
-					Title = $"Incorrect hash, got hash \"{e.SuppliedHash}\" but hash of content was determined to be \"{e.ContentHash}\""
+					Title = $"Incorrect hash, got hash \"{e.SuppliedHash}\" but hash of content was determined to be \"{e.ContentHash}\"",
+					Type = ProblemTypes.HashMismatchError
 				});
 			}
 			catch (ClientSendSlowException e)
@@ -1026,7 +1027,8 @@ namespace Jupiter.Controllers
 			{
 				return BadRequest(new ProblemDetails
 				{
-					Title = $"Incorrect hash, got hash \"{e.SuppliedHash}\" but hash of content was determined to be \"{e.ContentHash}\""
+					Title = $"Incorrect hash, got hash \"{e.SuppliedHash}\" but hash of content was determined to be \"{e.ContentHash}\"",
+					Type = ProblemTypes.HashMismatchError
 				});
 			}
 			catch (RemoteServiceException e)
@@ -1204,8 +1206,8 @@ namespace Jupiter.Controllers
 			{
 				return BadRequest(new ProblemDetails
 				{
-					Title =
-						$"Incorrect hash, got hash \"{e.SuppliedHash}\" but hash of content was determined to be \"{e.ContentHash}\""
+					Title = $"Incorrect hash, got hash \"{e.SuppliedHash}\" but hash of content was determined to be \"{e.ContentHash}\"",
+					Type = ProblemTypes.HashMismatchError
 				});
 			}
 			catch (ResourceHasToManyRequestsException)
