Cloud Storage - Lower log level for missing objects while checking build store consistency

#Synced-CL 48106205

diff --git a/Engine/Source/Programs/UnrealCloudDDC/Jupiter/Implementation/Builds/BuildStoreConsistencyCheckService.cs b/Engine/Source/Programs/UnrealCloudDDC/Jupiter/Implementation/Builds/BuildStoreConsistencyCheckService.cs
--- a/Engine/Source/Programs/UnrealCloudDDC/Jupiter/Implementation/Builds/BuildStoreConsistencyCheckService.cs
+++ b/Engine/Source/Programs/UnrealCloudDDC/Jupiter/Implementation/Builds/BuildStoreConsistencyCheckService.cs
@@ -173,12 +173,12 @@ namespace Jupiter.Implementation
 							}
 							catch (BlobNotFoundException e)
 							{
-								_logger.LogError(e, "Could not retrieve referenced blobs for {RefId} in {Namespace}.{Bucket}", refId, ns, bucket);
+								_logger.LogWarning(e, "Could not retrieve referenced blobs for {RefId} in {Namespace}.{Bucket}", refId, ns, bucket);
 								Interlocked.Increment(ref countOfBlobErrors);
 							}
 							catch (RefNotFoundException e)
 							{
-								_logger.LogError(e, "Could not retrieve referenced blobs for {RefId} in {Namespace}.{Bucket}", refId, ns, bucket);
+								_logger.LogWarning(e, "Could not retrieve referenced blobs for {RefId} in {Namespace}.{Bucket}", refId, ns, bucket);
 								Interlocked.Increment(ref countOfRefErrors);
 							}
 						}
