Cloud Storage - Catch scylla communication exception so that it doesn't hit the top level unhandled filter thus reducing monitoring noise. Loggning warning as its is a recoverable issue (consistency check will restart on the next poll)

#Synced-CL 49272873

diff --git a/Engine/Source/Programs/UnrealCloudDDC/Jupiter/Implementation/Blob/RefStoreConsistencyCheckService.cs b/Engine/Source/Programs/UnrealCloudDDC/Jupiter/Implementation/Blob/RefStoreConsistencyCheckService.cs
--- a/Engine/Source/Programs/UnrealCloudDDC/Jupiter/Implementation/Blob/RefStoreConsistencyCheckService.cs
+++ b/Engine/Source/Programs/UnrealCloudDDC/Jupiter/Implementation/Blob/RefStoreConsistencyCheckService.cs
@@ -6,6 +6,7 @@ using System.Linq;
 using System.Threading;
 using System.Threading.Tasks;
 using EpicGames.Horde.Storage;
+using Jupiter.Common;
 using Jupiter.Implementation.Blob;
 using Microsoft.Extensions.Logging;
 using Microsoft.Extensions.Options;
@@ -54,7 +55,14 @@ namespace Jupiter.Implementation
 				return false;
 			}
 
-			await RunConsistencyCheckAsync(cancellationToken);
+			try
+			{
+				await RunConsistencyCheckAsync(cancellationToken);
+			}
+			catch (RemoteServiceException e)
+			{
+				_logger.LogWarning(e, "Scylla communication issue occured while running consistency check");
+			}
 
 			return true;
 		}
