Cloud DDC - Added option to indicate that AWS Keyspace is being used and removed usage of the DISTINCT keyword in scylla requests when that is the case as its not supported. We hope that this will be sufficent to test running against keyspaces which is not something Epic intends to do ourselves but has been requested by licensees to experiment with.

#Synced-CL 49272873

diff --git a/Engine/Source/Programs/UnrealCloudDDC/Jupiter/Implementation/Objects/ScyllaReferencesStore.cs b/Engine/Source/Programs/UnrealCloudDDC/Jupiter/Implementation/Objects/ScyllaReferencesStore.cs
--- a/Engine/Source/Programs/UnrealCloudDDC/Jupiter/Implementation/Objects/ScyllaReferencesStore.cs
+++ b/Engine/Source/Programs/UnrealCloudDDC/Jupiter/Implementation/Objects/ScyllaReferencesStore.cs
@@ -90,8 +90,8 @@ namespace Jupiter.Implementation
 			string cqlOptions = scyllaSessionManager.IsScylla ? "BYPASS CACHE" : "";
 			_getObjectsStatement = _session.Prepare($"SELECT namespace, bucket, name, last_access_time FROM objects {cqlOptions}");
 			_getObjectsLastAccessStatement = _session.Prepare($"SELECT namespace, bucket, name, last_access_time FROM object_last_access_v2 {cqlOptions}");
-			_getNamespacesStatement = _session.Prepare("SELECT DISTINCT namespace FROM buckets_v2");
-			_getNamespacesOldStatement = _session.Prepare("SELECT DISTINCT namespace FROM buckets");
+			_getNamespacesStatement = scyllaSessionManager.IsAWSKeyspace ? _session.Prepare("SELECT namespace FROM buckets_v2") :  _session.Prepare("SELECT DISTINCT namespace FROM buckets_v2");
+			_getNamespacesOldStatement = scyllaSessionManager.IsAWSKeyspace ? _session.Prepare("SELECT namespace FROM buckets") : _session.Prepare("SELECT DISTINCT namespace FROM buckets");
 
 			_getObjectsBucketsStatement = _session.Prepare("select namespace, bucket FROM objects");
 			_getObjectsForPartitionRangeStatement = _session.Prepare($"SELECT namespace, bucket, name, last_access_time FROM objects WHERE token(namespace, bucket, name) >= ? AND token(namespace, bucket, name) <= ? {cqlOptions}");
@@ -548,14 +548,41 @@ namespace Jupiter.Implementation
 					throw new RemoteServiceException("Error communicating with Scylla", e);
 				}
 
-				foreach (Row row in rowSet)
+				if (_scyllaSessionManager.IsAWSKeyspace)
 				{
-					if (rowSet.GetAvailableWithoutFetching() == 0)
+					// keyspaces do not support the distinct keyword so we filter here on the client instead
+
+					List<NamespaceId> namespaces = new List<NamespaceId>();
+					foreach (Row row in rowSet)
 					{
-						await rowSet.FetchMoreResultsAsync();
+						if (rowSet.GetAvailableWithoutFetching() == 0)
+						{
+							await rowSet.FetchMoreResultsAsync();
+						}
+
+						NamespaceId ns = new NamespaceId(row.GetValue<string>(0));
+						if (!namespaces.Contains(ns))
+						{
+							namespaces.Add(ns);
+						}
 					}
 
-					yield return new NamespaceId(row.GetValue<string>(0));
+					foreach (NamespaceId namespaceId in namespaces)
+					{
+						yield return namespaceId;
+					}
+				}
+				else
+				{
+					foreach (Row row in rowSet)
+					{
+						if (rowSet.GetAvailableWithoutFetching() == 0)
+						{
+							await rowSet.FetchMoreResultsAsync();
+						}
+
+						yield return new NamespaceId(row.GetValue<string>(0));
+					}
 				}
 			}
 
diff --git a/Engine/Source/Programs/UnrealCloudDDC/Jupiter/Implementation/Objects/ScyllaSettings.cs b/Engine/Source/Programs/UnrealCloudDDC/Jupiter/Implementation/Objects/ScyllaSettings.cs
--- a/Engine/Source/Programs/UnrealCloudDDC/Jupiter/Implementation/Objects/ScyllaSettings.cs
+++ b/Engine/Source/Programs/UnrealCloudDDC/Jupiter/Implementation/Objects/ScyllaSettings.cs
@@ -57,6 +57,11 @@ public class ScyllaSettings : IValidatableObject
 	/// </summary>
 	public bool UseAzureCosmosDB { get; set; } = false;
 
+	/// <summary>
+	/// Whether a AWS Keyspace is used, adjusts requests to be compatible with this flavor of Cassandra.
+	/// </summary>
+	public bool UseAWSKeyspace { get; set; } = false;
+
 	/// <summary>
 	/// Set to false to disable use of SSL, this is only used when connecting to a CosmosDB instance
 	/// </summary>
diff --git a/Engine/Source/Programs/UnrealCloudDDC/Jupiter/Implementation/ScyllaSessionManager.cs b/Engine/Source/Programs/UnrealCloudDDC/Jupiter/Implementation/ScyllaSessionManager.cs
--- a/Engine/Source/Programs/UnrealCloudDDC/Jupiter/Implementation/ScyllaSessionManager.cs
+++ b/Engine/Source/Programs/UnrealCloudDDC/Jupiter/Implementation/ScyllaSessionManager.cs
@@ -20,6 +20,11 @@ namespace Jupiter.Implementation
 		/// True when using a base cassandra protocol
 		/// </summary>
 		bool IsCassandra { get; }
+
+		/// <summary>
+		/// True when AWS Keyspace is being used
+		/// </summary>
+		bool IsAWSKeyspace { get; }
 	}
 
 	public class ScyllaSessionManager : IScyllaSessionManager
@@ -27,10 +32,11 @@ namespace Jupiter.Implementation
 		private readonly ISession _replicatedSession;
 		private readonly ISession _localSession;
 
-		public ScyllaSessionManager(ISession replicatedSession, ISession localSession, bool isScylla, bool isCassandra)
+		public ScyllaSessionManager(ISession replicatedSession, ISession localSession, bool isScylla, bool isCassandra, bool isKeyspace)
 		{
 			_replicatedSession = replicatedSession;
 			_localSession = localSession;
+			IsAWSKeyspace = isKeyspace;
 			IsScylla = isScylla;
 			IsCassandra = isCassandra;
 		}
@@ -48,6 +54,7 @@ namespace Jupiter.Implementation
 
 		public bool IsScylla { get; }
 		public bool IsCassandra { get; }
+		public bool IsAWSKeyspace { get; }
 	}
 
 	public static class ScyllaUtils
diff --git a/Engine/Source/Programs/UnrealCloudDDC/Jupiter/Jupiter.Startup.cs b/Engine/Source/Programs/UnrealCloudDDC/Jupiter/Jupiter.Startup.cs
--- a/Engine/Source/Programs/UnrealCloudDDC/Jupiter/Jupiter.Startup.cs
+++ b/Engine/Source/Programs/UnrealCloudDDC/Jupiter/Jupiter.Startup.cs
@@ -477,8 +477,8 @@ namespace Jupiter
 			}
 			localSession.UserDefinedTypes.Define(UdtMap.For<ScyllaBlobIdentifier>("blob_identifier", localKeyspaceName));
 
-			bool isScylla = !settings.UseAzureCosmosDB;
-			return new ScyllaSessionManager(replicatedSession, localSession, isScylla, !isScylla);
+			bool isScylla = !settings.UseAzureCosmosDB && !settings.UseAWSKeyspace;
+			return new ScyllaSessionManager(replicatedSession, localSession, isScylla, !isScylla, settings.UseAWSKeyspace);
 		}
 
 		private IReferencesStore ObjectStoreFactory(IServiceProvider provider)
