Cloud Storage - Added DefaultConsistencyLevel to ScyllaSettings

#Synced-CL 48106205

diff --git a/Engine/Source/Programs/UnrealCloudDDC/Jupiter/Implementation/Objects/ScyllaSettings.cs b/Engine/Source/Programs/UnrealCloudDDC/Jupiter/Implementation/Objects/ScyllaSettings.cs
--- a/Engine/Source/Programs/UnrealCloudDDC/Jupiter/Implementation/Objects/ScyllaSettings.cs
+++ b/Engine/Source/Programs/UnrealCloudDDC/Jupiter/Implementation/Objects/ScyllaSettings.cs
@@ -3,6 +3,7 @@
 using System;
 using System.Collections.Generic;
 using System.ComponentModel.DataAnnotations;
+using Cassandra;
 
 namespace Jupiter.Implementation;
 
@@ -125,6 +126,11 @@ public class ScyllaSettings : IValidatableObject
 	/// </summary>
 	public bool CleanupOldTrackedState { get; set; } = true;
 
+	/// <summary>
+	/// Used to override the consistency level for requests
+	/// </summary>
+	public ConsistencyLevel DefaultConsistencyLevel { get; set; } = ConsistencyLevel.LocalOne;
+
 	/// <summary>
 	/// When enabled uses a more consistent read level for content ids which solves intermittent issues as the cost of worse performance
 	/// </summary>
diff --git a/Engine/Source/Programs/UnrealCloudDDC/Jupiter/Jupiter.Startup.cs b/Engine/Source/Programs/UnrealCloudDDC/Jupiter/Jupiter.Startup.cs
--- a/Engine/Source/Programs/UnrealCloudDDC/Jupiter/Jupiter.Startup.cs
+++ b/Engine/Source/Programs/UnrealCloudDDC/Jupiter/Jupiter.Startup.cs
@@ -364,7 +364,26 @@ namespace Jupiter
 				.WithLoadBalancingPolicy(Policies.NewDefaultLoadBalancingPolicy(settings.LocalDatacenterName))
 				.WithPoolingOptions(PoolingOptions.Create().SetMaxConnectionsPerHost(HostDistance.Local, settings.MaxConnectionForLocalHost).SetMaxRequestsPerConnection(settings.MaxRequestsPerConnection))
 				.WithExecutionProfiles(options =>
-					options.WithProfile("default", builder => builder.WithConsistencyLevel(ConsistencyLevel.LocalOne)).WithProfile("local_consistent", builder => builder.WithConsistencyLevel(ConsistencyLevel.LocalQuorum)));
+					options.WithProfile("default", builder => builder.WithConsistencyLevel(settings.DefaultConsistencyLevel))
+						.WithProfile("local_consistent", builder =>
+						{
+							ConsistencyLevel level;
+
+							// Use the DefaultConsistencyLevel if it provides higher consistency than LocalQuorum
+							switch (settings.DefaultConsistencyLevel)
+							{
+								case ConsistencyLevel.Quorum:
+								case ConsistencyLevel.EachQuorum:
+								case ConsistencyLevel.All:
+									level = settings.DefaultConsistencyLevel;
+									break;
+								default:
+									level = ConsistencyLevel.LocalQuorum;
+									break;
+							}
+							
+							builder.WithConsistencyLevel(level);
+						}));
 
 			if (settings.ReadTimeout != -1)
 			{
