Cloud Storage - Revert unnecessary AspNetCore.HealthChecks.Publisher.Datadog upgrade

#Synced-CL 48106205

diff --git a/Engine/Source/Programs/UnrealCloudDDC/Jupiter.Common/Jupiter.Common.csproj b/Engine/Source/Programs/UnrealCloudDDC/Jupiter.Common/Jupiter.Common.csproj
--- a/Engine/Source/Programs/UnrealCloudDDC/Jupiter.Common/Jupiter.Common.csproj
+++ b/Engine/Source/Programs/UnrealCloudDDC/Jupiter.Common/Jupiter.Common.csproj
@@ -25,7 +25,7 @@
   <ItemGroup>
     <PackageReference Include="Blake3" Version="1.1.0" />
 
-    <PackageReference Include="AspNetCore.HealthChecks.Publisher.Datadog" Version="9.0.0" />
+    <PackageReference Include="AspNetCore.HealthChecks.Publisher.Datadog" Version="8.0.0" />
     <PackageReference Include="AWSSDK.Extensions.NETCore.Setup" Version="3.7.300" />
     <PackageReference Include="AWSSDK.SecretsManager" Version="3.7.302.8" />
     <PackageReference Include="AWSSDK.SecurityToken" Version="3.7.300.34" />
