Cloud Storage - Remove AspNetCore.HealthChecks.Publisher.Datadog

#Synced-CL 48106205

diff --git a/Engine/Source/Programs/UnrealCloudDDC/Jupiter.Common/BaseStartup.cs b/Engine/Source/Programs/UnrealCloudDDC/Jupiter.Common/BaseStartup.cs
--- a/Engine/Source/Programs/UnrealCloudDDC/Jupiter.Common/BaseStartup.cs
+++ b/Engine/Source/Programs/UnrealCloudDDC/Jupiter.Common/BaseStartup.cs
@@ -475,12 +475,6 @@ namespace Jupiter
 			IHealthChecksBuilder healthChecks = services.AddHealthChecks()
 				.AddCheck("self", () => HealthCheckResult.Healthy(), tags: new[] { "self" });
 			OnAddHealthChecks(services, healthChecks);
-
-			string? ddAgentHost = System.Environment.GetEnvironmentVariable("DD_AGENT_HOST");
-			if (!string.IsNullOrEmpty(ddAgentHost))
-			{
-				healthChecks.AddDatadogPublisher("jupiter.healthchecks");
-			}
 		}
 
 		/// <summary>
diff --git a/Engine/Source/Programs/UnrealCloudDDC/Jupiter.Common/Implementation/InstrumentedStream.cs b/Engine/Source/Programs/UnrealCloudDDC/Jupiter.Common/Implementation/InstrumentedStream.cs
--- a/Engine/Source/Programs/UnrealCloudDDC/Jupiter.Common/Implementation/InstrumentedStream.cs
+++ b/Engine/Source/Programs/UnrealCloudDDC/Jupiter.Common/Implementation/InstrumentedStream.cs
@@ -1,6 +1,8 @@
 ﻿// Copyright Epic Games, Inc. All Rights Reserved.
 
 using System;
+using System.Collections.Generic;
+using System.Diagnostics.Metrics;
 using System.IO;
 using System.Threading;
 using System.Threading.Tasks;
@@ -18,10 +20,17 @@ namespace Jupiter.Common.Implementation
 		private long _totalBytesWritten = 0;
 		private DateTime _consumeStartedAt;
 
-		public InstrumentedStream(Stream s, string sourceIdentifier)
+		private static Histogram<double>? s_streamThroughputCounter;
+
+		public InstrumentedStream(Stream s, Meter meter, string sourceIdentifier)
 		{
 			_streamImplementation = s;
 			_sourceIdentifier = sourceIdentifier;
+
+			if (s_streamThroughputCounter == null)
+			{
+				s_streamThroughputCounter = meter.CreateHistogram<double>("jupiter.stream_throughput");
+			}
 		}
 
 		public override async ValueTask DisposeAsync()
@@ -36,7 +45,7 @@ namespace Jupiter.Common.Implementation
 			TimeSpan duration = DateTime.Now - _consumeStartedAt;
 			double rate = _totalBytesWritten / duration.TotalSeconds;
 
-			StatsdClient.DogStatsd.Histogram("jupiter.stream_throughput", rate, tags: new string[] { "sourceIdentifier:" + _sourceIdentifier });
+			s_streamThroughputCounter?.Record(rate, new KeyValuePair<string, object?>("sourceIdentifier", _sourceIdentifier));
 			base.Close();
 		}
 
diff --git a/Engine/Source/Programs/UnrealCloudDDC/Jupiter.Common/Jupiter.Common.csproj b/Engine/Source/Programs/UnrealCloudDDC/Jupiter.Common/Jupiter.Common.csproj
--- a/Engine/Source/Programs/UnrealCloudDDC/Jupiter.Common/Jupiter.Common.csproj
+++ b/Engine/Source/Programs/UnrealCloudDDC/Jupiter.Common/Jupiter.Common.csproj
@@ -25,7 +25,6 @@
   <ItemGroup>
     <PackageReference Include="Blake3" Version="1.1.0" />
 
-    <PackageReference Include="AspNetCore.HealthChecks.Publisher.Datadog" Version="8.0.0" />
     <PackageReference Include="AWSSDK.Extensions.NETCore.Setup" Version="3.7.300" />
     <PackageReference Include="AWSSDK.SecretsManager" Version="3.7.302.8" />
     <PackageReference Include="AWSSDK.SecurityToken" Version="3.7.300.34" />
diff --git a/Engine/Source/Programs/UnrealCloudDDC/Jupiter.Common/Utils/RequestUtil.cs b/Engine/Source/Programs/UnrealCloudDDC/Jupiter.Common/Utils/RequestUtil.cs
--- a/Engine/Source/Programs/UnrealCloudDDC/Jupiter.Common/Utils/RequestUtil.cs
+++ b/Engine/Source/Programs/UnrealCloudDDC/Jupiter.Common/Utils/RequestUtil.cs
@@ -1,6 +1,8 @@
 ﻿// Copyright Epic Games, Inc. All Rights Reserved.
 
 using System;
+using System.Collections.Generic;
+using System.Diagnostics.Metrics;
 using System.IO;
 using System.Threading.Tasks;
 using Microsoft.AspNetCore.Http;
@@ -11,6 +13,8 @@ namespace Jupiter
 {
 	public static class RequestUtil
 	{
+		private static Histogram<double>? s_streamThroughputCounter;
+		
 		public static async Task<byte[]> ReadRawBodyAsync(HttpRequest request)
 		{
 			long? contentLength = request.GetTypedHeaders().ContentLength;
@@ -31,7 +35,14 @@ namespace Jupiter
 			TimeSpan duration = DateTime.Now - readStart;
 
 			double rate = contentLength.Value / duration.TotalSeconds;
-			StatsdClient.DogStatsd.Histogram("jupiter.stream_throughput", rate, tags: new string[] { "sourceIdentifier:" + "readbody" });
+
+			if (s_streamThroughputCounter == null)
+			{
+				Meter meter = request.HttpContext.RequestServices.GetService<Meter>()!;
+				s_streamThroughputCounter = meter.CreateHistogram<double>("jupiter.stream_throughput");
+			}
+
+			s_streamThroughputCounter.Record(rate, new KeyValuePair<string, object?>("sourceIdentifier", "readbody"));
 
 			return ms.GetBuffer();
 		}
diff --git a/Engine/Source/Programs/UnrealCloudDDC/Jupiter/Jupiter.Startup.cs b/Engine/Source/Programs/UnrealCloudDDC/Jupiter/Jupiter.Startup.cs
--- a/Engine/Source/Programs/UnrealCloudDDC/Jupiter/Jupiter.Startup.cs
+++ b/Engine/Source/Programs/UnrealCloudDDC/Jupiter/Jupiter.Startup.cs
@@ -28,7 +28,6 @@ using Microsoft.Extensions.Configuration;
 using Microsoft.Extensions.DependencyInjection;
 using Microsoft.Extensions.Logging;
 using Microsoft.Extensions.Options;
-using StatsdClient;
 
 namespace Jupiter
 {
@@ -37,18 +36,7 @@ namespace Jupiter
 	{
 		public JupiterStartup(IConfiguration configuration) : base(configuration, CreateLogger<JupiterStartup>())
 		{
-			string? ddAgentHost = System.Environment.GetEnvironmentVariable("DD_AGENT_HOST");
-			if (!string.IsNullOrEmpty(ddAgentHost))
-			{
-				Logger.LogInformation("Initializing Dogstatsd to connect to: {DatadogAgentHost}", ddAgentHost);
-				StatsdConfig dogstatsdConfig = new StatsdConfig
-				{
-					StatsdServerName = ddAgentHost,
-					StatsdPort = 8125,
-				};
 
-				DogStatsd.Configure(dogstatsdConfig);
-			}
 		}
 
 		protected override void OnAddAuthorization(AuthorizationOptions authorizationOptions, List<string> defaultSchemes)
