Cloud DDC - Allow the use of the ssl settings when connecting to AWS Keyspace

#Synced-CL 49760190

diff --git a/Engine/Source/Programs/UnrealCloudDDC/Jupiter/Implementation/Objects/ScyllaSettings.cs b/Engine/Source/Programs/UnrealCloudDDC/Jupiter/Implementation/Objects/ScyllaSettings.cs
--- a/Engine/Source/Programs/UnrealCloudDDC/Jupiter/Implementation/Objects/ScyllaSettings.cs
+++ b/Engine/Source/Programs/UnrealCloudDDC/Jupiter/Implementation/Objects/ScyllaSettings.cs
@@ -64,7 +64,7 @@ public class ScyllaSettings : IValidatableObject
 	public bool UseAWSKeyspace { get; set; } = false;
 
 	/// <summary>
-	/// Set to false to disable use of SSL, this is only used when connecting to a CosmosDB instance
+	/// Set to false to disable use of SSL, this is only used when connecting to a CosmosDB instance or a AWS Keyspace instance
 	/// </summary>
 	public bool UseSSL { get; set; } = true;
 
diff --git a/Engine/Source/Programs/UnrealCloudDDC/Jupiter/Jupiter.Startup.cs b/Engine/Source/Programs/UnrealCloudDDC/Jupiter/Jupiter.Startup.cs
--- a/Engine/Source/Programs/UnrealCloudDDC/Jupiter/Jupiter.Startup.cs
+++ b/Engine/Source/Programs/UnrealCloudDDC/Jupiter/Jupiter.Startup.cs
@@ -377,7 +377,7 @@ namespace Jupiter
 			{
 				clusterBuilder.SocketOptions.SetReadTimeoutMillis(settings.ReadTimeout);
 			}
-			if (settings.UseAzureCosmosDB)
+			if (settings.UseAzureCosmosDB || settings.UseAWSKeyspace)
 			{
 				CassandraConnectionStringBuilder connectionStringBuilder = new CassandraConnectionStringBuilder(connectionString);
 				string[] contactPoints = connectionStringBuilder.ContactPoints;
