Cloud Storage - Upgrade OpenTelemetry dependencies to 1.14.0

#Synced-CL 48106205

diff --git a/Engine/Source/Programs/UnrealCloudDDC/Jupiter.Common/Jupiter.Common.csproj b/Engine/Source/Programs/UnrealCloudDDC/Jupiter.Common/Jupiter.Common.csproj
--- a/Engine/Source/Programs/UnrealCloudDDC/Jupiter.Common/Jupiter.Common.csproj
+++ b/Engine/Source/Programs/UnrealCloudDDC/Jupiter.Common/Jupiter.Common.csproj
@@ -31,13 +31,13 @@
     <PackageReference Include="Azure.Identity" Version="1.11.4" />
     <PackageReference Include="Azure.Security.KeyVault.Secrets" Version="4.5.0" />
     <PackageReference Include="Newtonsoft.Json" Version="13.0.3" />
-    <PackageReference Include="OpenTelemetry" Version="1.7.0" />
-    <PackageReference Include="OpenTelemetry.Exporter.Console" Version="1.7.0" />
-    <PackageReference Include="OpenTelemetry.Exporter.OpenTelemetryProtocol" Version="1.7.0" />
-    <PackageReference Include="OpenTelemetry.Instrumentation.AspNetCore" Version="1.8.1" />
-    <PackageReference Include="OpenTelemetry.Instrumentation.Http" Version="1.8.1" />
-    <PackageReference Include="OpenTelemetry.Instrumentation.Runtime" Version="1.7.0" />
-    <PackageReference Include="OpenTelemetry.Extensions.Hosting" Version="1.7.0" />
+    <PackageReference Include="OpenTelemetry" Version="1.14.0" />
+    <PackageReference Include="OpenTelemetry.Exporter.Console" Version="1.14.0" />
+    <PackageReference Include="OpenTelemetry.Exporter.OpenTelemetryProtocol" Version="1.14.0" />
+    <PackageReference Include="OpenTelemetry.Instrumentation.AspNetCore" Version="1.14.0" />
+    <PackageReference Include="OpenTelemetry.Instrumentation.Http" Version="1.14.0" />
+    <PackageReference Include="OpenTelemetry.Instrumentation.Runtime" Version="1.14.0" />
+    <PackageReference Include="OpenTelemetry.Extensions.Hosting" Version="1.14.0" />
     <PackageReference Include="IdentityServer4.AccessTokenValidation" Version="3.0.1" />
 
     <PackageReference Include="Microsoft.AspNet.WebApi.Client" Version="6.0.0" />
