Cloud Storage - If a user tries to download a build but they lack access to all parts we now return a 403 (forbidden) as they practically lack access to the entire build, rather then a filtered out response with nothing in it which breaks other tooling.

#Synced-CL 50250157

diff --git a/Engine/Source/Programs/UnrealCloudDDC/Jupiter/Controllers/BuildsController.cs b/Engine/Source/Programs/UnrealCloudDDC/Jupiter/Controllers/BuildsController.cs
--- a/Engine/Source/Programs/UnrealCloudDDC/Jupiter/Controllers/BuildsController.cs
+++ b/Engine/Source/Programs/UnrealCloudDDC/Jupiter/Controllers/BuildsController.cs
@@ -372,11 +372,13 @@ namespace Jupiter.Controllers
 				// inject the build parts into the object
 				// we know the object fields will all have the same type so make it a uniform object to please compact binary validation
 				newObjectWriter.BeginUniformObject("parts", CbFieldType.ObjectId | CbFieldType.HasFieldName);
+				bool hadAccessToPart = false;
 				await foreach ((string partName, CbObjectId partId) in _buildStore.GetBuildPartsAsync(ns, bucket, buildId))
 				{
 					ActionResult? accessResultPart = await _requestHelper.HasAccessToScopeAsync(User, Request, new AccessScope(ns, bucket, BuildPartType, partName), new[] { JupiterAclAction.ReadBuildPart });
 					if (accessResultPart == null)
 					{
+						hadAccessToPart = true;
 						// only include the build part if this request has access to it
 						newObjectWriter.WriteObjectId(partName, partId);
 					}
@@ -384,7 +386,13 @@ namespace Jupiter.Controllers
 				newObjectWriter.EndObject(); // parts
 				newObjectWriter.EndObject(); // root
 				CbObject newObject = newObjectWriter.ToObject();
-				
+
+				if (!hadAccessToPart)
+				{
+					// user did not have access to any part, that means they had no access to this at all
+					return new ForbidResult();
+				}
+
 				string responseType = _formatResolver.GetResponseType(Request, format, CustomMediaTypeNames.UnrealCompactBinary);
 				Tracer.CurrentSpan.SetAttribute("response-type", responseType);
 
