UnrealCloudDCC: Update dockerfile to copy property sheet

#rnx

#Synced-CL 50830039

diff --git a/Engine/Source/Programs/UnrealCloudDDC/.dockerignore b/Engine/Source/Programs/UnrealCloudDDC/.dockerignore
--- a/Engine/Source/Programs/UnrealCloudDDC/.dockerignore
+++ b/Engine/Source/Programs/UnrealCloudDDC/.dockerignore
@@ -5,5 +5,7 @@
 
 **/bin
 **/obj
+**/uatbin
+**/uatobj
 **/out
 **/node_modules
\ No newline at end of file
diff --git a/Engine/Source/Programs/UnrealCloudDDC/dockerfile b/Engine/Source/Programs/UnrealCloudDDC/dockerfile
--- a/Engine/Source/Programs/UnrealCloudDDC/dockerfile
+++ b/Engine/Source/Programs/UnrealCloudDDC/dockerfile
@@ -7,6 +7,7 @@ WORKDIR /app
 ARG UE_P4_CHANGELIST=0
 
 # Copy csprojs for nuget restore
+COPY Programs/Shared/* Programs/Shared/
 COPY Programs/Shared/EpicGames.Core/*.csproj ./Programs/Shared/EpicGames.Core/
 COPY Programs/Shared/EpicGames.AspNet/*.csproj ./Programs/Shared/EpicGames.AspNet/
 COPY Programs/Shared/EpicGames.IoHash/*.csproj ./Programs/Shared/EpicGames.IoHash/
