Cloud Storage - Default to UnrealCompactBinary response types for symbol retrieval

#Synced-CL 48106205

diff --git a/Engine/Source/Programs/UnrealCloudDDC/Jupiter/Controllers/SymbolsController.cs b/Engine/Source/Programs/UnrealCloudDDC/Jupiter/Controllers/SymbolsController.cs
--- a/Engine/Source/Programs/UnrealCloudDDC/Jupiter/Controllers/SymbolsController.cs
+++ b/Engine/Source/Programs/UnrealCloudDDC/Jupiter/Controllers/SymbolsController.cs
@@ -111,7 +111,7 @@ namespace Jupiter.Controllers
 				using ServerTimingMetricScoped? serverTimingScope =
 					serverTiming?.CreateServerTimingMetricScope("body.write", "Time spent writing body");
 
-				string responseType = _formatResolver.GetResponseType(Request, null, MediaTypeNames.Application.Octet);
+				string responseType = _formatResolver.GetResponseType(Request, null, CustomMediaTypeNames.UnrealCompactBinary);
 				Tracer.CurrentSpan.SetAttribute("response-type", responseType);
 				
 				using TelemetrySpan scope = _tracer.StartActiveSpan("body.write").SetAttribute("operation.name", "body.write");
