Cloud DDC - Do not include stacktraces on missing blobs for the batch endpoint. While these are raised as exceptions they are not errors and we never have any use of the stacktrace which is very expensive to build.

#Synced-CL 50250157

diff --git a/Engine/Source/Programs/UnrealCloudDDC/Jupiter/Controllers/ReferencesController.cs b/Engine/Source/Programs/UnrealCloudDDC/Jupiter/Controllers/ReferencesController.cs
--- a/Engine/Source/Programs/UnrealCloudDDC/Jupiter/Controllers/ReferencesController.cs
+++ b/Engine/Source/Programs/UnrealCloudDDC/Jupiter/Controllers/ReferencesController.cs
@@ -1300,7 +1300,8 @@ namespace Jupiter.Controllers
 				}
 				catch (Exception ex) when (ex is RefNotFoundException or PartialReferenceResolveException or ReferenceIsMissingBlobsException)
 				{
-					return ToErrorResult(ex, HttpStatusCode.NotFound);
+					// this isn't really an error so no need to build a stacktrace which is very expensive
+					return ToErrorResult(ex, HttpStatusCode.NotFound, includeStacktrace: false);
 				}
 				catch (Exception e)
 				{
@@ -1424,14 +1425,14 @@ namespace Jupiter.Controllers
 			});
 		}
 
-		private static (CbObject, HttpStatusCode) ToErrorResult(Exception exception, HttpStatusCode statusCode = HttpStatusCode.InternalServerError)
+		private static (CbObject, HttpStatusCode) ToErrorResult(Exception exception, HttpStatusCode statusCode = HttpStatusCode.InternalServerError, bool includeStacktrace = true)
 		{
 			Exception e = exception;
 			CbWriter writer = new CbWriter();
 			writer.BeginObject();
 			writer.WriteString("title", e.Message);
 			writer.WriteInteger("status", (int)statusCode);
-			if (e.StackTrace != null)
+			if (includeStacktrace && e.StackTrace != null)
 			{
 				writer.WriteString("stackTrace", e.StackTrace);
 			}
