Cloud Storage - Return empty lists instead of 403s when authenticated users cannot see any namespaces or buckets

#Synced-CL 50915067

diff --git a/Engine/Source/Programs/UnrealCloudDDC/Jupiter/Controllers/BuildsController.cs b/Engine/Source/Programs/UnrealCloudDDC/Jupiter/Controllers/BuildsController.cs
--- a/Engine/Source/Programs/UnrealCloudDDC/Jupiter/Controllers/BuildsController.cs
+++ b/Engine/Source/Programs/UnrealCloudDDC/Jupiter/Controllers/BuildsController.cs
@@ -2428,11 +2428,6 @@ namespace Jupiter.Controllers
 				}
 			}
 
-			if (!namespacesWithAccess.Any())
-			{
-				return new ForbidResult();
-			}
-
 			LogAuditEntry(HttpMethod.Get, nameof(GetNamespacesAsync), null, null, null, null);
 
 			return Ok(new GetNamespacesResponse(namespacesWithAccess.ToArray()));
@@ -2442,6 +2437,13 @@ namespace Jupiter.Controllers
 		[ProducesDefaultResponseType]
 		public async Task<IActionResult> GetBucketsAsync(NamespaceId ns)
 		{
+			ActionResult? nsAccessResult = await _requestHelper.HasAccessToScopeAsync(User, Request, new AccessScope(ns), new[] { JupiterAclAction.EnumerateBucket });
+
+			if (nsAccessResult != null)
+			{
+				return new ForbidResult();
+			}
+			
 			LogAuditEntry(HttpMethod.Get, nameof(GetBucketsAsync), ns, null, null, null);
 
 			BucketId[] potentialBuckets = await _refService.GetBucketsAsync(ns, HttpContext.RequestAborted).ToArrayAsync(HttpContext.RequestAborted);
@@ -2456,10 +2458,6 @@ namespace Jupiter.Controllers
 					verifiedBuckets.Add(bucket);
 				}
 			}
-			if (!verifiedBuckets.Any())
-			{
-				return new ForbidResult();
-			}
 
 			return Ok(new GetBucketsResponse(verifiedBuckets.ToArray()));
 		}
