Cloud Storage - Add custom properties when AWS Keyspace is used to enable TTL. Note that we are unable to test these ourselves so this is provided as is and may not function correctly.

#Synced-CL 51251964

diff --git a/Engine/Source/Programs/UnrealCloudDDC/Jupiter/Implementation/Builds/ScyllaBuildStore.cs b/Engine/Source/Programs/UnrealCloudDDC/Jupiter/Implementation/Builds/ScyllaBuildStore.cs
--- a/Engine/Source/Programs/UnrealCloudDDC/Jupiter/Implementation/Builds/ScyllaBuildStore.cs
+++ b/Engine/Source/Programs/UnrealCloudDDC/Jupiter/Implementation/Builds/ScyllaBuildStore.cs
@@ -29,6 +29,8 @@ public class ScyllaBuildStore : IBuildStore
 		_session = scyllaSessionManager.GetSessionForReplicatedKeyspace();
 		_mapper = new Mapper(_session);
 
+		string tableOptions = scyllaSessionManager.IsAWSKeyspace ? "WITH CUSTOM_PROPERTIES={'ttl':{'status':'enabled'}}" : "";
+		string appendTableOptions = scyllaSessionManager.IsAWSKeyspace ? "AND CUSTOM_PROPERTIES={'ttl':{'status':'enabled'}}" : "";
 		if (!scyllaSettings.CurrentValue.AvoidSchemaChanges)
 		{
 			_session.Execute(new SimpleStatement(@$"CREATE TABLE IF NOT EXISTS builds (
@@ -38,7 +40,7 @@ public class ScyllaBuildStore : IBuildStore
 					is_finalized boolean,
 					build_object blob,
 					PRIMARY KEY ((namespace, bucket_id), build_id)
-				) WITH CLUSTERING ORDER BY (build_id DESC)
+				) WITH CLUSTERING ORDER BY (build_id DESC) {appendTableOptions}
 				;"
 			));
 
@@ -49,7 +51,7 @@ public class ScyllaBuildStore : IBuildStore
 					part_name text,
 					part_id blob,
 					PRIMARY KEY ((namespace, bucket_id, build_id), part_id)
-				);"
+				) {tableOptions};"
 			));
 
 			_session.Execute(new SimpleStatement(@$"CREATE TABLE IF NOT EXISTS build_appendixes (
@@ -58,7 +60,7 @@ public class ScyllaBuildStore : IBuildStore
 					build_id blob,
 					appendix_name text,
 					PRIMARY KEY ((namespace, bucket_id, build_id), appendix_name)
-				);"
+				) {tableOptions};"
 			));
 		}
 		_getTtlStatement = _session.Prepare("SELECT TTL(build_object) as ttl FROM builds WHERE namespace = ? AND bucket_id = ? AND build_id = ?");
diff --git a/Engine/Source/Programs/UnrealCloudDDC/Jupiter/Implementation/TransactionLog/ScyllaReplicationLog.cs b/Engine/Source/Programs/UnrealCloudDDC/Jupiter/Implementation/TransactionLog/ScyllaReplicationLog.cs
--- a/Engine/Source/Programs/UnrealCloudDDC/Jupiter/Implementation/TransactionLog/ScyllaReplicationLog.cs
+++ b/Engine/Source/Programs/UnrealCloudDDC/Jupiter/Implementation/TransactionLog/ScyllaReplicationLog.cs
@@ -26,9 +26,10 @@ namespace Jupiter.Implementation
 			_session = scyllaSessionManager.GetSessionForLocalKeyspace();
 			_mapper = new Mapper(_session);
 
+			string suffix = scyllaSessionManager.IsAWSKeyspace ? "WITH CUSTOM_PROPERTIES={'ttl':{'status':'enabled'}}" : "";
 			if (!settings.CurrentValue.AvoidSchemaChanges)
 			{
-				_session.Execute(new SimpleStatement(@"CREATE TABLE IF NOT EXISTS replication_log (
+				_session.Execute(new SimpleStatement(@$"CREATE TABLE IF NOT EXISTS replication_log (
 					namespace varchar,
 					replication_bucket bigint,
 					replication_id timeuuid,
@@ -37,10 +38,10 @@ namespace Jupiter.Implementation
 					type int,
 					object_identifier blob_identifier,
 					PRIMARY KEY ((namespace, replication_bucket), replication_id)
-				);"
+				) {suffix};"
 				));
 
-				_session.Execute(new SimpleStatement(@"CREATE TABLE IF NOT EXISTS blob_replication_log (
+				_session.Execute(new SimpleStatement(@$"CREATE TABLE IF NOT EXISTS blob_replication_log (
 					namespace varchar,
 					replication_bucket bigint,
 					replication_id timeuuid,
@@ -48,7 +49,7 @@ namespace Jupiter.Implementation
 					blob_id blob,
 					bucket_hint varchar,
 					PRIMARY KEY ((namespace, replication_bucket), replication_id)
-				);"
+				) {suffix};"
 				));
 
 				_session.Execute(new SimpleStatement(@"CREATE TABLE IF NOT EXISTS replication_snapshot (
