Cloud Storage - Increased max scanned builds to 40k (from 20k) as we have yet to implement proper pagination and thus need to let this scan a lot (we are currently seeing warnings on partial results due to this)

#Synced-CL 51251964

diff --git a/Engine/Source/Programs/UnrealCloudDDC/Jupiter/Controllers/BuildsController.cs b/Engine/Source/Programs/UnrealCloudDDC/Jupiter/Controllers/BuildsController.cs
--- a/Engine/Source/Programs/UnrealCloudDDC/Jupiter/Controllers/BuildsController.cs
+++ b/Engine/Source/Programs/UnrealCloudDDC/Jupiter/Controllers/BuildsController.cs
@@ -2721,7 +2721,7 @@ namespace Jupiter.Controllers
 		public int Skip { get; set; } = 0;
 
 		[CbField("max")]
-		public int Max { get; set; } = 20_000;
+		public int Max { get; set; } = 40_000;
 
 		[CbField("maxTimeMS")]
 		public uint MaxTimeMs { get; set; } = Debugger.IsAttached ? 0u : 120_000; // default to 2 minutes max per search unless debugging in which case we disable it
