Cloud Storage - Only redirect to nginx if the file actually exists

#Synced-CL 51251964

diff --git a/Engine/Source/Programs/UnrealCloudDDC/Jupiter/Implementation/NginxRedirectHelper.cs b/Engine/Source/Programs/UnrealCloudDDC/Jupiter/Implementation/NginxRedirectHelper.cs
--- a/Engine/Source/Programs/UnrealCloudDDC/Jupiter/Implementation/NginxRedirectHelper.cs
+++ b/Engine/Source/Programs/UnrealCloudDDC/Jupiter/Implementation/NginxRedirectHelper.cs
@@ -1,6 +1,7 @@
 // Copyright Epic Games, Inc. All Rights Reserved.
 
 using System;
+using System.IO;
 using System.Threading.Tasks;
 using Microsoft.AspNetCore.Http;
 using Microsoft.AspNetCore.Mvc;
@@ -47,6 +48,10 @@ namespace Jupiter.Implementation
 				return false;
 			}
 
+			if (!File.Exists(blobContents.LocalPath))
+			{
+				return false;
+			}
 			return true;
 		}
 	}
