Cloud Storage - Suppress warning about file access when verifying if nginx will be able to send the file

#Synced-CL 51251964

diff --git a/Engine/Source/Programs/UnrealCloudDDC/Jupiter/Implementation/NginxRedirectHelper.cs b/Engine/Source/Programs/UnrealCloudDDC/Jupiter/Implementation/NginxRedirectHelper.cs
--- a/Engine/Source/Programs/UnrealCloudDDC/Jupiter/Implementation/NginxRedirectHelper.cs
+++ b/Engine/Source/Programs/UnrealCloudDDC/Jupiter/Implementation/NginxRedirectHelper.cs
@@ -48,10 +48,13 @@ namespace Jupiter.Implementation
 				return false;
 			}
 
+			// suppress warning about accepting paths from user requests while this is technically based on a blob name send by the user that name has to be a hash and the actual path is built from the hash
+#pragma warning disable CA3003 			
 			if (!File.Exists(blobContents.LocalPath))
 			{
 				return false;
 			}
+#pragma warning restore CA3003
 			return true;
 		}
 	}
