Cloud Storage - Return a error if we attempt to serve a multi range request via asp.net (which does not support it and returns the full object).

Also fixed issue with nginx redirect checking for file existence using a relative path which does not work

#Synced-CL 51251964

diff --git a/Engine/Source/Programs/UnrealCloudDDC/Jupiter.Common/Implementation/BlobContents.cs b/Engine/Source/Programs/UnrealCloudDDC/Jupiter.Common/Implementation/BlobContents.cs
--- a/Engine/Source/Programs/UnrealCloudDDC/Jupiter.Common/Implementation/BlobContents.cs
+++ b/Engine/Source/Programs/UnrealCloudDDC/Jupiter.Common/Implementation/BlobContents.cs
@@ -12,12 +12,13 @@ namespace Jupiter.Implementation
 		private Stream? _stream;
 		private readonly IBufferedPayload? _bufferedPayload;
 
-		public BlobContents(Stream stream, long length, string? localPath = null)
+		public BlobContents(Stream stream, long length, string? localPath = null, string? absolutePath = null)
 		{
 			_stream = stream;
 
 			Length = length;
 			LocalPath = localPath;
+			AbsolutePath = absolutePath;
 		}
 
 		public BlobContents(byte[] payload)
@@ -60,6 +61,7 @@ namespace Jupiter.Implementation
 		public long Length { get; }
 
 		public string? LocalPath { get; }
+		public string? AbsolutePath { get; }
 		
 		public IBufferedPayload? BufferedPayload => _bufferedPayload;
 
diff --git a/Engine/Source/Programs/UnrealCloudDDC/Jupiter/Controllers/BuildsController.cs b/Engine/Source/Programs/UnrealCloudDDC/Jupiter/Controllers/BuildsController.cs
--- a/Engine/Source/Programs/UnrealCloudDDC/Jupiter/Controllers/BuildsController.cs
+++ b/Engine/Source/Programs/UnrealCloudDDC/Jupiter/Controllers/BuildsController.cs
@@ -1830,6 +1830,11 @@ namespace Jupiter.Controllers
 					return _nginxRedirectHelper.CreateActionResult(blobContents, mediaType);
 				}
 
+				IActionResult? response = _requestHelper.EnsureSingleRangeRequest(HttpContext);
+				if (response != null)
+				{
+					return response;
+				}
 				return File(blobContents.Stream, mediaType, enableRangeProcessing: true);
 			}
 			catch (BlobNotFoundException e)
diff --git a/Engine/Source/Programs/UnrealCloudDDC/Jupiter/Controllers/IRequestHelper.cs b/Engine/Source/Programs/UnrealCloudDDC/Jupiter/Controllers/IRequestHelper.cs
--- a/Engine/Source/Programs/UnrealCloudDDC/Jupiter/Controllers/IRequestHelper.cs
+++ b/Engine/Source/Programs/UnrealCloudDDC/Jupiter/Controllers/IRequestHelper.cs
@@ -13,6 +13,7 @@ namespace Jupiter.Controllers
 		public Task<ActionResult?> HasAccessToScopeAsync(ClaimsPrincipal user, HttpRequest request, AccessScope scope, JupiterAclAction[] aclActions);
 
 		public bool IsPublicPort(HttpContext context);
+		IActionResult? EnsureSingleRangeRequest(HttpContext context);
 	}
 
 	public class AuthorizationException : Exception
diff --git a/Engine/Source/Programs/UnrealCloudDDC/Jupiter/Controllers/RequestHelper.cs b/Engine/Source/Programs/UnrealCloudDDC/Jupiter/Controllers/RequestHelper.cs
--- a/Engine/Source/Programs/UnrealCloudDDC/Jupiter/Controllers/RequestHelper.cs
+++ b/Engine/Source/Programs/UnrealCloudDDC/Jupiter/Controllers/RequestHelper.cs
@@ -2,6 +2,7 @@
 
 using System;
 using System.Linq;
+using System.Net;
 using System.Security.Claims;
 using System.Threading.Tasks;
 using Jupiter.Common;
@@ -189,4 +190,15 @@ public class RequestHelper : IRequestHelper
 
 		return isPublicPort;
 	}
+
+	public IActionResult? EnsureSingleRangeRequest(HttpContext context)
+	{
+		// asp.net does not support serving multi range requests so we return an error should that be requested
+		if (context.Request.Headers.Range.ToString().Contains(',', StringComparison.OrdinalIgnoreCase))
+		{
+			return new StatusCodeResult((int)HttpStatusCode.NotImplemented);
+		}
+
+		return null;
+	}
 }
diff --git a/Engine/Source/Programs/UnrealCloudDDC/Jupiter/Implementation/Blob/FileSystemStore.cs b/Engine/Source/Programs/UnrealCloudDDC/Jupiter/Implementation/Blob/FileSystemStore.cs
--- a/Engine/Source/Programs/UnrealCloudDDC/Jupiter/Implementation/Blob/FileSystemStore.cs
+++ b/Engine/Source/Programs/UnrealCloudDDC/Jupiter/Implementation/Blob/FileSystemStore.cs
@@ -516,7 +516,7 @@ namespace Jupiter.Implementation
 			}
 			FileStream fs = new FileStream(filePath.FullName, FileMode.Open, FileAccess.Read, FileShare.Read, DefaultBufferSize, FileOptions.Asynchronous | FileOptions.SequentialScan);
 
-			return Task.FromResult<BlobContents?>(new BlobContents(fs, fs.Length, $"{_namespaceId}/{path}"));
+			return Task.FromResult<BlobContents?>(new BlobContents(fs, fs.Length, localPath: $"{_namespaceId}/{path}", absolutePath: filePath.FullName));
 		}
 
 		public Task<bool> ExistsAsync(string path, CancellationToken cancellationToken = default)
diff --git a/Engine/Source/Programs/UnrealCloudDDC/Jupiter/Implementation/NginxRedirectHelper.cs b/Engine/Source/Programs/UnrealCloudDDC/Jupiter/Implementation/NginxRedirectHelper.cs
--- a/Engine/Source/Programs/UnrealCloudDDC/Jupiter/Implementation/NginxRedirectHelper.cs
+++ b/Engine/Source/Programs/UnrealCloudDDC/Jupiter/Implementation/NginxRedirectHelper.cs
@@ -49,8 +49,8 @@ namespace Jupiter.Implementation
 			}
 
 			// suppress warning about accepting paths from user requests while this is technically based on a blob name send by the user that name has to be a hash and the actual path is built from the hash
-#pragma warning disable CA3003 			
-			if (!File.Exists(blobContents.LocalPath))
+#pragma warning disable CA3003 
+			if (!File.Exists(blobContents.AbsolutePath))
 			{
 				return false;
 			}
