Cloud Storage - Avoid doing file exists checks that uses the internal FileInfo cache within the FileSystemStore . This could fix the occasional issues we have seen were blobs that should be available to serve return 404 while the checks before it claimed the blob existed.

#Synced-CL 51823794

diff --git a/Engine/Source/Programs/UnrealCloudDDC/Jupiter/Implementation/Blob/FileSystemStore.cs b/Engine/Source/Programs/UnrealCloudDDC/Jupiter/Implementation/Blob/FileSystemStore.cs
--- a/Engine/Source/Programs/UnrealCloudDDC/Jupiter/Implementation/Blob/FileSystemStore.cs
+++ b/Engine/Source/Programs/UnrealCloudDDC/Jupiter/Implementation/Blob/FileSystemStore.cs
@@ -446,11 +446,12 @@ namespace Jupiter.Implementation
 			FileInfo filePath = GetFilesystemPath(path);
 			filePath.Directory?.Create();
 
-			if (!filePath.Exists)
+			if (!File.Exists(filePath.FullName))
 			{
 				int uniqueId = Interlocked.Increment(ref s_uniqueId);
 
 				string tempFilePath = $"{filePath.FullName}.{s_processSuffix}.{uniqueId}";
+
 				await using (FileStream fs = new FileStream(tempFilePath, FileMode.Create, FileAccess.Write, FileShare.Read, DefaultBufferSize, FileOptions.Asynchronous | FileOptions.SequentialScan))
 				{
 					await content.CopyToAsync(fs, cancellationToken);
@@ -505,7 +506,7 @@ namespace Jupiter.Implementation
 		{
 			FileInfo filePath = GetFilesystemPath(path);
 
-			if (!filePath.Exists)
+			if (!File.Exists(filePath.FullName))
 			{
 				return Task.FromResult<BlobContents?>(null);
 			}
@@ -529,7 +530,7 @@ namespace Jupiter.Implementation
 		public Task DeleteAsync(string path, CancellationToken cancellationToken = default)
 		{
 			FileInfo filePath = GetFilesystemPath(path);
-			if (filePath.Exists)
+			if (File.Exists(filePath.FullName))
 			{
 				filePath.Delete();
 			}
