Cloud Storage - Fixed another case of potentially stale blob existance information was used

#Synced-CL 51823794

diff --git a/Engine/Source/Programs/UnrealCloudDDC/Jupiter/Implementation/Blob/FileSystemStore.cs b/Engine/Source/Programs/UnrealCloudDDC/Jupiter/Implementation/Blob/FileSystemStore.cs
--- a/Engine/Source/Programs/UnrealCloudDDC/Jupiter/Implementation/Blob/FileSystemStore.cs
+++ b/Engine/Source/Programs/UnrealCloudDDC/Jupiter/Implementation/Blob/FileSystemStore.cs
@@ -524,7 +524,7 @@ namespace Jupiter.Implementation
 		{
 			FileInfo filePath = GetFilesystemPath(path);
 
-			return Task.FromResult(filePath.Exists);
+			return Task.FromResult(File.Exists(filePath.FullName));
 		}
 
 		public Task DeleteAsync(string path, CancellationToken cancellationToken = default)
