Cloud Storage - Add optional header (X-Jupiter-Force-Consistent-Read) to force consistent reads, only used in ref get right now. This can be used by the editor when attempting to read back values that does read after write (with the KeepAlive policy)

#Synced-CL 51823794

diff --git a/Engine/Source/Programs/UnrealCloudDDC/Jupiter/Controllers/ReferencesController.cs b/Engine/Source/Programs/UnrealCloudDDC/Jupiter/Controllers/ReferencesController.cs
--- a/Engine/Source/Programs/UnrealCloudDDC/Jupiter/Controllers/ReferencesController.cs
+++ b/Engine/Source/Programs/UnrealCloudDDC/Jupiter/Controllers/ReferencesController.cs
@@ -120,13 +120,15 @@ namespace Jupiter.Controllers
 		/// <param name="bucket">The category/type of record you are caching. Is a clustered key together with the actual key, but all records in the same bucket can be dropped easily. Example: `terrainTexture` </param>
 		/// <param name="key">The unique name of this particular key. `iAmAVeryValidKey`</param>
 		/// <param name="format">Optional specifier to set which output format is used json/raw/cb</param>
+		/// <param name="forceConsistentRead">Optional header to force consistent reads disabling all uses of caches</param>
 		[HttpGet("{ns}/{bucket}/{key}.{format?}", Order = 500)]
 		[Produces(MediaTypeNames.Application.Json, MediaTypeNames.Application.Octet, CustomMediaTypeNames.UnrealCompactBinary, CustomMediaTypeNames.JupiterInlinedPayload, CustomMediaTypeNames.UnrealCompactBinaryPackage)]
 		public async Task<IActionResult> GetAsync(
 			[FromRoute][Required] NamespaceId ns,
 			[FromRoute][Required] BucketId bucket,
 			[FromRoute][Required] RefId key,
-			[FromRoute] string? format = null)
+			[FromRoute] string? format = null,
+			[FromHeader(Name = CommonHeaders.ForceConsistentRead)] bool forceConsistentRead = false)
 		{
 			ActionResult? accessResult = await _requestHelper.HasAccessToScopeAsync(User, Request, new AccessScope(ns, bucket), new[] { JupiterAclAction.ReadObject });
 			if (accessResult != null)
@@ -139,7 +141,7 @@ namespace Jupiter.Controllers
 				Tracer.CurrentSpan.SetAttribute("bucket", bucket.ToString());
 				Tracer.CurrentSpan.SetAttribute("namespace", ns.ToString());
 
-				(RefRecord objectRecord, BlobContents? maybeBlob) = await _refService.GetAsync(ns, bucket, key, Array.Empty<string>());
+				(RefRecord objectRecord, BlobContents? maybeBlob) = await _refService.GetAsync(ns, bucket, key, Array.Empty<string>(), consistentRead: forceConsistentRead);
 
 				if (maybeBlob == null)
 				{
diff --git a/Engine/Source/Programs/UnrealCloudDDC/Jupiter/Implementation/CommonHeaders.cs b/Engine/Source/Programs/UnrealCloudDDC/Jupiter/Implementation/CommonHeaders.cs
--- a/Engine/Source/Programs/UnrealCloudDDC/Jupiter/Implementation/CommonHeaders.cs
+++ b/Engine/Source/Programs/UnrealCloudDDC/Jupiter/Implementation/CommonHeaders.cs
@@ -12,5 +12,6 @@ namespace Jupiter.Implementation
 		public const string InlinePayloadHash = "X-Jupiter-InlinePayloadHash";
 
 		public const string AllowOverwrite = "X-Jupiter-Allow-Overwrite";
+		public const string ForceConsistentRead = "X-Jupiter-Force-Consistent-Read";
 	}
 }
diff --git a/Engine/Source/Programs/UnrealCloudDDC/Jupiter/Implementation/IRefService.cs b/Engine/Source/Programs/UnrealCloudDDC/Jupiter/Implementation/IRefService.cs
--- a/Engine/Source/Programs/UnrealCloudDDC/Jupiter/Implementation/IRefService.cs
+++ b/Engine/Source/Programs/UnrealCloudDDC/Jupiter/Implementation/IRefService.cs
@@ -11,7 +11,7 @@ namespace Jupiter.Implementation
 {
 	public interface IRefService
 	{
-		Task<(RefRecord, BlobContents?)> GetAsync(NamespaceId ns, BucketId bucket, RefId key, string[]? fields = null, bool doLastAccessTracking = true, CancellationToken cancellationToken = default);
+		Task<(RefRecord, BlobContents?)> GetAsync(NamespaceId ns, BucketId bucket, RefId key, string[]? fields = null, bool consistentRead = false, bool doLastAccessTracking = true, CancellationToken cancellationToken = default);
 
 		Task<(ContentId[], BlobId[])> PutAsync(NamespaceId ns, BucketId bucket, RefId key, BlobId blobHash, CbObject payload, Action<BlobId>? onBlobFound = null, bool allowOverwrite = false, CancellationToken cancellationToken = default);
 		Task<(ContentId[], BlobId[])> FinalizeAsync(NamespaceId ns, BucketId bucket, RefId key, BlobId blobHash, Action<BlobId>? onBlobFound = null, CancellationToken cancellationToken = default);
diff --git a/Engine/Source/Programs/UnrealCloudDDC/Jupiter/Implementation/RefService.cs b/Engine/Source/Programs/UnrealCloudDDC/Jupiter/Implementation/RefService.cs
--- a/Engine/Source/Programs/UnrealCloudDDC/Jupiter/Implementation/RefService.cs
+++ b/Engine/Source/Programs/UnrealCloudDDC/Jupiter/Implementation/RefService.cs
@@ -47,13 +47,7 @@ namespace Jupiter.Implementation
 			_cloudDDCSettings = cloudDDCSettings;
 		}
 
-		public Task<(RefRecord, BlobContents?)> GetAsync(NamespaceId ns, BucketId bucket, RefId key, string[]? fields = null, bool doLastAccessTracking = true, CancellationToken cancellationToken = default)
-		{
-			return GetAsync(ns, bucket, key, fields, doLastAccessTracking, skipCache: false, cancellationToken: cancellationToken);
-		}
-
-		// ReSharper disable once MethodOverloadWithOptionalParameter - this private overload exists only for bypassing cache for internal use within this service
-		private async Task<(RefRecord, BlobContents?)> GetAsync(NamespaceId ns, BucketId bucket, RefId key, string[]? fields = null, bool doLastAccessTracking = true, bool skipCache = false, CancellationToken cancellationToken = default)
+		public async Task<(RefRecord, BlobContents?)> GetAsync(NamespaceId ns, BucketId bucket, RefId key, string[]? fields = null, bool consistentRead = false, bool doLastAccessTracking = true, CancellationToken cancellationToken = default)
 		{
 			// if no field filtering is being used we assume everything is needed
 			IReferencesStore.FieldFlags flags = IReferencesStore.FieldFlags.All;
@@ -73,7 +67,7 @@ namespace Jupiter.Implementation
 			}
 
 			IReferencesStore.OperationFlags opFlags = IReferencesStore.OperationFlags.None;
-			if (skipCache)
+			if (consistentRead)
 			{
 				opFlags |= IReferencesStore.OperationFlags.BypassCache;
 			}
@@ -201,7 +195,7 @@ namespace Jupiter.Implementation
 		public async Task<(ContentId[], BlobId[])> FinalizeAsync(NamespaceId ns, BucketId bucket, RefId key, BlobId blobHash, Action<BlobId>? onBlobFound = null, CancellationToken cancellationToken = default)
 		{
 			// finalize is intended to verify the state of the object in the db, so we bypass any caches to make sure the object we are working on is not stale
-			(RefRecord o, BlobContents? blob) = await GetAsync(ns, bucket, key, skipCache: true, cancellationToken: cancellationToken);
+			(RefRecord o, BlobContents? blob) = await GetAsync(ns, bucket, key, consistentRead: true, cancellationToken: cancellationToken);
 			if (blob == null)
 			{
 				throw new InvalidOperationException("No blob when attempting to finalize");
@@ -311,7 +305,7 @@ namespace Jupiter.Implementation
 		{
 			try
 			{
-				(RefRecord, BlobContents?) _ = await GetAsync(ns, bucket, key, new string[] { "name" }, doLastAccessTracking: false, skipCache: true, cancellationToken: cancellationToken);
+				(RefRecord, BlobContents?) _ = await GetAsync(ns, bucket, key, new string[] { "name" }, doLastAccessTracking: false, consistentRead: true, cancellationToken: cancellationToken);
 			}
 			catch (NamespaceNotFoundException)
 			{
