Cloud Storage - Temporarily disable copying in-memory blob contents to file system

#Synced-CL 51557602

diff --git a/Engine/Source/Programs/UnrealCloudDDC/Jupiter/Implementation/Blob/BlobService.cs b/Engine/Source/Programs/UnrealCloudDDC/Jupiter/Implementation/Blob/BlobService.cs
--- a/Engine/Source/Programs/UnrealCloudDDC/Jupiter/Implementation/Blob/BlobService.cs
+++ b/Engine/Source/Programs/UnrealCloudDDC/Jupiter/Implementation/Blob/BlobService.cs
@@ -445,21 +445,21 @@ public class BlobService : IBlobService
 				}
 			}
 
-			{
-				if (numStoreMisses == 1 && _blobStores[0] is FileSystemStore fileSystemStore)
-				{
-					// if we are only missing data from the filesystem store, buffer to it and then send using the data in the filesystem so that we can use sendfile for this data
-					using TelemetrySpan scope = _tracer.StartActiveSpan("HierarchicalStore.PopulateStore")
-						.SetAttribute("operation.name", "HierarchicalStore.PopulateStore")
-						.SetAttribute("Mode", "UseFilesystemPayload")
-						.SetAttribute("BlobStore", nameof(FileSystemStore));
-
-					await using Stream s = blobContents.Stream;
-					await fileSystemStore.PutObjectAsync(ns, s, blob, cancellationToken);
-
-					return await fileSystemStore.GetObjectAsync(ns, blob, flags: LastAccessTrackingFlags.DoTracking, cancellationToken: cancellationToken);
-				}
-			}
+			// {
+			// 	if (numStoreMisses == 1 && _blobStores[0] is FileSystemStore fileSystemStore)
+			// 	{
+			// 		// if we are only missing data from the filesystem store, buffer to it and then send using the data in the filesystem so that we can use sendfile for this data
+			// 		using TelemetrySpan scope = _tracer.StartActiveSpan("HierarchicalStore.PopulateStore")
+			// 			.SetAttribute("operation.name", "HierarchicalStore.PopulateStore")
+			// 			.SetAttribute("Mode", "UseFilesystemPayload")
+			// 			.SetAttribute("BlobStore", nameof(FileSystemStore));
+			//
+			// 		await using Stream s = blobContents.Stream;
+			// 		await fileSystemStore.PutObjectAsync(ns, s, blob, cancellationToken);
+			//
+			// 		return await fileSystemStore.GetObjectAsync(ns, blob, flags: LastAccessTrackingFlags.DoTracking, cancellationToken: cancellationToken);
+			// 	}
+			// }
 			
 			bool streamCanSeek = blobContents.Stream.CanSeek;
 			
