Cloud Storage - Temporarily disable returning 501 responses for multi-range requests

#Synced-CL 51557602

diff --git a/Engine/Source/Programs/UnrealCloudDDC/Jupiter/Controllers/BuildsController.cs b/Engine/Source/Programs/UnrealCloudDDC/Jupiter/Controllers/BuildsController.cs
--- a/Engine/Source/Programs/UnrealCloudDDC/Jupiter/Controllers/BuildsController.cs
+++ b/Engine/Source/Programs/UnrealCloudDDC/Jupiter/Controllers/BuildsController.cs
@@ -1830,11 +1830,11 @@ namespace Jupiter.Controllers
 					return _nginxRedirectHelper.CreateActionResult(blobContents, mediaType);
 				}
 
-				IActionResult? response = _requestHelper.EnsureSingleRangeRequest(HttpContext);
-				if (response != null)
-				{
-					return response;
-				}
+				// IActionResult? response = _requestHelper.EnsureSingleRangeRequest(HttpContext);
+				// if (response != null)
+				// {
+				// 	return response;
+				// }
 				return File(blobContents.Stream, mediaType, enableRangeProcessing: true);
 			}
 			catch (BlobNotFoundException e)
