Cloud Storage - Clean up PDB parsing

#Synced-CL 51557602

diff --git a/Engine/Source/Programs/UnrealCloudDDC/Jupiter/Controllers/SymbolsController.cs b/Engine/Source/Programs/UnrealCloudDDC/Jupiter/Controllers/SymbolsController.cs
--- a/Engine/Source/Programs/UnrealCloudDDC/Jupiter/Controllers/SymbolsController.cs
+++ b/Engine/Source/Programs/UnrealCloudDDC/Jupiter/Controllers/SymbolsController.cs
@@ -492,40 +492,71 @@ namespace Jupiter.Controllers
 			int freeBlockMapIndex = reader.ReadInt32();
 			int blockCount = reader.ReadInt32();
 			int directoryStreamLength = reader.ReadInt32();
-			int reserved = reader.ReadInt32();
+			int _ = reader.ReadInt32();
 			int hintBlock = reader.ReadInt32();
 
-			// read the free block maps
+			switch (blockSize)
+			{
+				case 512:
+				case 1024:
+				case 2048:
+				case 4096:
+					break;
+				default:
+					throw new Exception($"Unsupported block size: {blockSize}");
+			}
+
+			if (freeBlockMapIndex != 1 && freeBlockMapIndex != 2)
+			{
+				throw new Exception("Unexpected free block map index");
+			}
+
+			if (blockCount * blockSize != reader.BaseStream.Length)
+			{
+				throw new Exception("Unexpected pdb stream length");
+			}
+
+			// build the stream directory
 			int freeBlockOffset = blockSize * hintBlock;
 			reader.BaseStream.Seek(freeBlockOffset, SeekOrigin.Begin);
-			int directoryBlock = reader.ReadInt32();
+			
+			int directoryStreamBlockCount = (int)Math.Ceiling((double)directoryStreamLength / blockSize);
+			int[] directoryBlocks = new int[directoryStreamBlockCount];
+			for (int i = 0; i < directoryStreamBlockCount; ++i)
+			{
+				directoryBlocks[i] = reader.ReadInt32();
+			}
+			
+			byte[] streamDirectory = new byte[directoryStreamBlockCount * blockSize];
+			int streamDirectoryOffset = 0;
 
+			foreach (int directoryBlock in directoryBlocks)
+			{
+				reader.BaseStream.Seek(blockSize * directoryBlock, SeekOrigin.Begin);
+				streamDirectoryOffset += reader.Read(streamDirectory, streamDirectoryOffset, blockSize);
+			}
+			
 			// read the stream directory
-			int streamDirectoryOffset = blockSize * directoryBlock;
-			reader.BaseStream.Seek(streamDirectoryOffset, SeekOrigin.Begin);
-			int streamCount = reader.ReadInt32();
+			using Stream streamDirectoryStream = new MemoryStream(streamDirectory);
+			using BinaryReader streamDirectoryReader = new BinaryReader(streamDirectoryStream);
+			int streamCount = streamDirectoryReader.ReadInt32();
 
 			int[] streamLengths = new int[streamCount];
 			for (int i = 0; i < streamCount; ++i)
 			{
-				streamLengths[i] = reader.ReadInt32();
+				streamLengths[i] = streamDirectoryReader.ReadInt32();
 			}
 
 			// calculate the blocks for each size
 			List<List<int>> blocks = new List<List<int>>();
 			for (int i = 0; i < streamCount; ++i)
 			{
-				int extraBlocks = streamLengths[i] / blockSize;
+				int streamBlockCount = (int)Math.Ceiling((double)streamLengths[i] / blockSize);
 
 				List<int> blocklist = new List<int>();
-				if (streamLengths[i] % blockSize > 0)
+				for (int j = 0; j < streamBlockCount; ++j)
 				{
-					blocklist.Add(reader.ReadInt32());
-				}
-
-				for (int j = 0; j < extraBlocks; ++j)
-				{
-					blocklist.Add(reader.ReadInt32());
+					blocklist.Add(streamDirectoryReader.ReadInt32());
 				}
 
 				blocks.Add(blocklist);
@@ -548,13 +579,10 @@ namespace Jupiter.Controllers
 				int destinationIndex = 0;
 
 				List<int> streamBlocks = blocks[i];
-				for (int index = 0; index < streamBlocks.Count - 1; ++index)
+				foreach (int streamBlock in streamBlocks)
 				{
-					reader.BaseStream.Seek(streamBlocks[index] * blockSize, SeekOrigin.Begin);
-					byte[] buf = reader.ReadBytes(blockSize);
-
-					Array.Copy(buf, 0, streamBuffer, destinationIndex, buf.Length);
-					destinationIndex += blockSize;
+					reader.BaseStream.Seek(streamBlock * blockSize, SeekOrigin.Begin);
+					destinationIndex += reader.Read(streamBuffer, destinationIndex, Math.Min(blockSize, streamBuffer.Length - destinationIndex));
 				}
 
 				// all blocks are combined into a contiguous stream
