Cloud Storage - Allow for custom init scripts

#Synced-CL 51557602

diff --git a/Engine/Source/Programs/UnrealCloudDDC/Helm/UnrealCloudDDC/templates/configmap-init.yaml b/Engine/Source/Programs/UnrealCloudDDC/Helm/UnrealCloudDDC/templates/configmap-init.yaml
new file mode 100644
--- /dev/null
+++ b/Engine/Source/Programs/UnrealCloudDDC/Helm/UnrealCloudDDC/templates/configmap-init.yaml
@@ -0,0 +1,26 @@
+{{- if .Values.prepareNvmeFilesystem.enabled -}}
+{{- $drivesCount := len .Values.prepareNvmeFilesystem.drives -}}
+{{- if and (not .Values.prepareNvmeFilesystem.customScript) (eq $drivesCount 0) }}{{ fail "You have to specify at least one drive path when enabling prepareNvmeFilesystem" }}{{- end -}}
+apiVersion: v1
+kind: ConfigMap
+metadata:
+  name: {{ include "common.names.fullname" . }}-init
+  labels: {{- include "common.labels.standard" . | nindent 4 }}
+data:
+  init.sh: |    
+    {{- if .Values.prepareNvmeFilesystem.customScript }}
+    {{ tpl .Values.prepareNvmeFilesystem.customScript . | nindent 4 }}
+    {{- else }}
+    #!/bin/sh
+    
+    # Exit if the node has already mounted the target
+    mountpoint -q /data-nvme && exit 0
+    
+    {{- if gt $drivesCount 1 }}
+    apt-get update && apt-get -y install mdadm && mdadm --create --verbose /dev/md0 --level=0 --raid-devices={{ $drivesCount }} {{ join " " .Values.prepareNvmeFilesystem.drives }} && mkfs.ext4 -E nodiscard /dev/md0 && mount /dev/md0 /data-nvme
+    {{- else }}
+    mkfs.ext4 -E nodiscard {{ first .Values.prepareNvmeFilesystem.drives }} && mount {{ first .Values.prepareNvmeFilesystem.drives }} /data-nvme
+    {{- end }}
+    
+    {{- end }}
+{{- end }}
diff --git a/Engine/Source/Programs/UnrealCloudDDC/Helm/UnrealCloudDDC/templates/deployment.yaml b/Engine/Source/Programs/UnrealCloudDDC/Helm/UnrealCloudDDC/templates/deployment.yaml
--- a/Engine/Source/Programs/UnrealCloudDDC/Helm/UnrealCloudDDC/templates/deployment.yaml
+++ b/Engine/Source/Programs/UnrealCloudDDC/Helm/UnrealCloudDDC/templates/deployment.yaml
@@ -58,28 +58,18 @@ spec:
       {{- include "common.tplvalues.render" (dict "value" .Values.initContainers "context" $) | nindent 6 }}
       {{- end }}
       {{- if .Values.prepareNvmeFilesystem.enabled }}
-      {{- $drivesCount := len .Values.prepareNvmeFilesystem.drives }}
       - name: nvme-mount
         image: ubuntu
         securityContext:
           privileged: true
-        command: ["/bin/sh"]
-      {{- if eq $drivesCount 0 }}
-      {{ fail "You have to specify at least one drivepath to prepare when enabling the prepareNvmeFilesystem" }}
-      {{- else if gt $drivesCount 1 }}
-        args:
-        - -c
-        - "apt-get update && apt-get -y install mdadm && mdadm --create --verbose /dev/md0 --level=0 --raid-devices={{ $drivesCount }} {{ join " " .Values.prepareNvmeFilesystem.drives }} && mkfs.ext4 -E nodiscard /dev/md0 && mount /dev/md0 /data-nvme || exit 0"
-      {{- else }}
-        args:
-        - -c
-        - mkfs.ext4 -E nodiscard {{ first .Values.prepareNvmeFilesystem.drives }} && mount {{ first .Values.prepareNvmeFilesystem.drives }} /data-nvme || exit 0
-      {{- end }}
-      {{- end }}
+        command: ["/bin/sh", "/init-scripts/init.sh"]
         volumeMounts:
         - mountPath: /data-nvme
           name: data
           mountPropagation: Bidirectional
+        - mountPath: /init-scripts
+          name: init-scripts
+          readOnly: true
       {{- end }}
       containers:
       - name: {{ include "common.names.name" . }}
@@ -294,6 +284,12 @@ spec:
         {{- include "common.tplvalues.render" ( dict "value" .Values.persistence.volume "context" $) | nindent 8 }}
         {{- end }}
       {{- end }}
+      {{- if .Values.prepareNvmeFilesystem.enabled }}
+      - name: init-scripts
+        configMap:
+          name: {{ $fullname }}-init
+          defaultMode: 0755
+      {{- end }}
       {{- if .Values.extraVolumes }}
       {{- include "common.tplvalues.render" ( dict "value" .Values.extraVolumes "context" $) | nindent 6 }}
       {{- end }}
diff --git a/Engine/Source/Programs/UnrealCloudDDC/Helm/UnrealCloudDDC/values.yaml b/Engine/Source/Programs/UnrealCloudDDC/Helm/UnrealCloudDDC/values.yaml
--- a/Engine/Source/Programs/UnrealCloudDDC/Helm/UnrealCloudDDC/values.yaml
+++ b/Engine/Source/Programs/UnrealCloudDDC/Helm/UnrealCloudDDC/values.yaml
@@ -550,4 +550,11 @@ worker:
 ## if more then 1 drive is specified we raid0 these together into a single drive
 prepareNvmeFilesystem:
   enabled: false
-  drives: []
\ No newline at end of file
+  drives: []
+  ## Optional: provide a fully custom init script instead of the default mkfs/mount logic.
+  ## Supports Go/Helm templates — .Values, .Release, etc. are all available.
+  ## Example:
+  ##   customScript: |
+  ##     mkfs.ext4 -E nodiscard {{ "{{" }} first .Values.prepareNvmeFilesystem.drives {{ "}}" }} && \
+  ##     mount {{ "{{" }} first .Values.prepareNvmeFilesystem.drives {{ "}}" }} /data-nvme || exit 0
+  customScript: ""
\ No newline at end of file
