Cloud Storage - Fix hanging conditional in Helm template

#Synced-CL 51557602

diff --git a/Engine/Source/Programs/UnrealCloudDDC/Helm/UnrealCloudDDC/templates/deployment.yaml b/Engine/Source/Programs/UnrealCloudDDC/Helm/UnrealCloudDDC/templates/deployment.yaml
--- a/Engine/Source/Programs/UnrealCloudDDC/Helm/UnrealCloudDDC/templates/deployment.yaml
+++ b/Engine/Source/Programs/UnrealCloudDDC/Helm/UnrealCloudDDC/templates/deployment.yaml
@@ -71,6 +71,7 @@ spec:
           name: init-scripts
           readOnly: true
       {{- end }}
+      {{- end }}
       containers:
       - name: {{ include "common.names.name" . }}
         image: {{ printf "%s:%s" .Values.image.repository (include "jupiter.appversion" . ) }}
