Cloud Storage - Byte array comparer used to sort keys in Data Store

#Synced-CL 51823794

diff --git a/Engine/Source/Programs/UnrealCloudDDC/Jupiter.Common/Utils/Sha1Utils.cs b/Engine/Source/Programs/UnrealCloudDDC/Jupiter.Common/Utils/Sha1Utils.cs
--- a/Engine/Source/Programs/UnrealCloudDDC/Jupiter.Common/Utils/Sha1Utils.cs
+++ b/Engine/Source/Programs/UnrealCloudDDC/Jupiter.Common/Utils/Sha1Utils.cs
@@ -1,5 +1,6 @@
 ﻿// Copyright Epic Games, Inc. All Rights Reserved.
 
+using System;
 using System.Collections.Generic;
 using System.Security.Cryptography;
 using System.Threading.Tasks;
@@ -33,7 +34,7 @@ namespace Jupiter
 		}
 	}
 
-	public class ByteArrayComparer : EqualityComparer<byte[]>
+	public class ByteArrayComparer : EqualityComparer<byte[]>, IComparer<byte[]>
 	{
 		public override bool Equals(byte[]? a, byte[]? b)
 		{
@@ -71,5 +72,32 @@ namespace Jupiter
 
 			return unchecked((int)b);
 		}
+
+		public int Compare(byte[]? x, byte[]? y)
+		{
+			if (x == null)
+			{
+				return -1;
+			}
+
+			if (y == null)
+			{
+				return 1;
+			}
+
+			for (int i = 0; i < x.Length; i++)
+			{
+				byte a = x[i];
+				byte b = y[i];
+
+				int comp = a.CompareTo(b);
+				if (comp != 0)
+				{
+					return comp;
+				}
+			}
+
+			return 0;
+		}
 	}
 }
