Cloud Storage - Use TTL to invalidate known buckets before inserting into Scylla

#Synced-CL 51557602

diff --git a/Engine/Source/Programs/UnrealCloudDDC/Jupiter/Implementation/Objects/ScyllaReferencesStore.cs b/Engine/Source/Programs/UnrealCloudDDC/Jupiter/Implementation/Objects/ScyllaReferencesStore.cs
--- a/Engine/Source/Programs/UnrealCloudDDC/Jupiter/Implementation/Objects/ScyllaReferencesStore.cs
+++ b/Engine/Source/Programs/UnrealCloudDDC/Jupiter/Implementation/Objects/ScyllaReferencesStore.cs
@@ -34,7 +34,7 @@ namespace Jupiter.Implementation
 		private readonly PreparedStatement _getObjectsBucketsStatement;
 		private readonly PreparedStatement _getObjectsLastAccessForPartitionRangeStatement;
 
-		private readonly ConcurrentDictionary<NamespaceId, ConcurrentBag<BucketId>> _addedBuckets = new ConcurrentDictionary<NamespaceId, ConcurrentBag<BucketId>>();
+		private readonly ConcurrentDictionary<NamespaceId, ConcurrentDictionary<BucketId, DateTime>> _addedBuckets = new ConcurrentDictionary<NamespaceId, ConcurrentDictionary<BucketId, DateTime>>();
 
 		public ScyllaReferencesStore(IScyllaSessionManager scyllaSessionManager, IOptionsMonitor<ScyllaSettings> settings, INamespacePolicyResolver namespacePolicyResolver, Tracer tracer, ILogger<ScyllaReferencesStore> logger)
 		{
@@ -736,22 +736,27 @@ namespace Jupiter.Implementation
 		{
 			using TelemetrySpan scope = _tracer.BuildScyllaSpan("scylla.add_bucket");
 
-			ConcurrentBag<BucketId> addedBuckets = _addedBuckets.GetOrAdd(ns, id => new ConcurrentBag<BucketId>());
+			ConcurrentDictionary<BucketId, DateTime> addedBuckets = _addedBuckets.GetOrAdd(ns, id => new ConcurrentDictionary<BucketId, DateTime>());
 
-			bool alreadyAdded = addedBuckets.Contains(bucket);
-			if (!alreadyAdded)
+			DateTime now = DateTime.UtcNow;
+			if (addedBuckets.TryGetValue(bucket, out DateTime added))
 			{
-				try
+				if (added > now - TimeSpan.FromHours(1))
 				{
-					Task addTask = _mapper.InsertAsync<ScyllaBucket>(new ScyllaBucket(ns, bucket));
-					await addTask;
-					addedBuckets.Add(bucket);
-				}
-				catch (QueryExecutionException e)
-				{
-					throw new RemoteServiceException("Error communicating with Scylla", e);
+					return;
 				}
 			}
+
+			try
+			{
+				Task addTask = _mapper.InsertAsync<ScyllaBucket>(new ScyllaBucket(ns, bucket));
+				await addTask;
+				addedBuckets.AddOrUpdate(bucket, now, (_, _) => now);
+			}
+			catch (QueryExecutionException e)
+			{
+				throw new RemoteServiceException("Error communicating with Scylla", e);
+			}
 		}
 
 		public async Task CleanupTrackedStateAsync()
