Cloud Storage - Add "symbols" ACL resource type

#Synced-CL 51557602

diff --git a/Engine/Source/Programs/UnrealCloudDDC/Jupiter/Controllers/SymbolsController.cs b/Engine/Source/Programs/UnrealCloudDDC/Jupiter/Controllers/SymbolsController.cs
--- a/Engine/Source/Programs/UnrealCloudDDC/Jupiter/Controllers/SymbolsController.cs
+++ b/Engine/Source/Programs/UnrealCloudDDC/Jupiter/Controllers/SymbolsController.cs
@@ -50,6 +50,11 @@ namespace Jupiter.Controllers
 		private readonly Tracer _tracer;
 		private readonly ILogger<SymbolsController> _logger;
 		private readonly ILogger? _auditLogger;
+		
+		/// <summary>
+		/// Constant used for access control of symbols
+		/// </summary>
+		private const string SymbolsType = "symbols";
 
 		public SymbolsController(IRefService refService, IBlobService blobStore, IDiagnosticContext diagnosticContext, BufferedPayloadFactory bufferedPayloadFactory, FormatResolver formatResolver, NginxRedirectHelper nginxRedirectHelper, IRequestHelper requestHelper, INamespacePolicyResolver namespacePolicyResolver, Tracer tracer, ILogger<SymbolsController> logger, IOptionsMonitor<SymbolsSettings> symbolSettings)
 		{
@@ -83,7 +88,7 @@ namespace Jupiter.Controllers
 			[FromRoute][Required] string fileName)
 		{
 			BucketId bucket = new BucketId(moduleName);
-			ActionResult? accessResult = await _requestHelper.HasAccessToScopeAsync(User, Request, new AccessScope(ns, bucket), new[] { JupiterAclAction.ReadObject });
+			ActionResult? accessResult = await _requestHelper.HasAccessToScopeAsync(User, Request, new AccessScope(ns, bucket, SymbolsType, moduleName), new[] { JupiterAclAction.ReadObject });
 			if (accessResult != null)
 			{
 				return accessResult;
@@ -209,7 +214,7 @@ namespace Jupiter.Controllers
 			[FromRoute][Required] string moduleName)
 		{
 			BucketId bucket = new BucketId(moduleName);
-			ActionResult? accessResult = await _requestHelper.HasAccessToScopeAsync(User, Request, new AccessScope(ns, bucket), new[] { JupiterAclAction.WriteObject });
+			ActionResult? accessResult = await _requestHelper.HasAccessToScopeAsync(User, Request, new AccessScope(ns, bucket, SymbolsType, moduleName), new[] { JupiterAclAction.WriteObject });
 			if (accessResult != null)
 			{
 				return accessResult;
@@ -331,7 +336,7 @@ namespace Jupiter.Controllers
 			[FromQuery] bool supportsRedirect = false)
 		{
 			BucketId bucket = new BucketId(moduleName);
-			ActionResult? result = await _requestHelper.HasAccessToScopeAsync(User, Request, new AccessScope(ns, bucket), new[] { JupiterAclAction.ReadObject });
+			ActionResult? result = await _requestHelper.HasAccessToScopeAsync(User, Request, new AccessScope(ns, bucket, SymbolsType, moduleName), new[] { JupiterAclAction.ReadObject });
 			if (result != null)
 			{
 				return result;
@@ -388,7 +393,7 @@ namespace Jupiter.Controllers
 			[Required] BlobId id)
 		{
 			BucketId bucket = new BucketId(moduleName);
-			ActionResult? result = await _requestHelper.HasAccessToScopeAsync(User, Request, new AccessScope(ns, bucket), new[] { JupiterAclAction.WriteObject });
+			ActionResult? result = await _requestHelper.HasAccessToScopeAsync(User, Request, new AccessScope(ns, bucket, SymbolsType, moduleName), new[] { JupiterAclAction.WriteObject });
 			if (result != null)
 			{
 				return result;
diff --git a/Engine/Source/Programs/UnrealCloudDDC/Jupiter/appsettings.Testing.json b/Engine/Source/Programs/UnrealCloudDDC/Jupiter/appsettings.Testing.json
--- a/Engine/Source/Programs/UnrealCloudDDC/Jupiter/appsettings.Testing.json
+++ b/Engine/Source/Programs/UnrealCloudDDC/Jupiter/appsettings.Testing.json
@@ -189,6 +189,31 @@
 					]
 				}
 			},
+			{
+				"Claims": [ "*" ],
+				"Actions": [
+					"ReadObject",
+					"WriteObject"
+				],
+				"Scope": {
+					"Namespaces": [
+						{
+							"values": [ "test-symbols-memory", "test-symbols-scylla" ]
+						}
+					],
+					"Buckets": [
+						{
+							"values": [ "*" ]
+						}
+					],
+					"Resources": [
+						{
+							"values": [ "*" ],
+							"resourceType": "symbols"
+						}
+					]
+				}
+			},
 			{
 				// global operations run in all namespaces and need the admin action
 				// this should really be protected with a specific claim but for tests we disable auth so we do not have claims to control on
