[NuGet] Update Newtonsoft.Json to 13.0.4

#rnx
#Synced-CL 52688981

diff --git a/Engine/Source/Programs/UnrealCloudDDC/Jupiter.Common/Jupiter.Common.csproj b/Engine/Source/Programs/UnrealCloudDDC/Jupiter.Common/Jupiter.Common.csproj
--- a/Engine/Source/Programs/UnrealCloudDDC/Jupiter.Common/Jupiter.Common.csproj
+++ b/Engine/Source/Programs/UnrealCloudDDC/Jupiter.Common/Jupiter.Common.csproj
@@ -30,7 +30,7 @@
     <PackageReference Include="AWSSDK.SecurityToken" Version="3.7.300.34" />
     <PackageReference Include="Azure.Identity" Version="1.11.4" />
     <PackageReference Include="Azure.Security.KeyVault.Secrets" Version="4.5.0" />
-    <PackageReference Include="Newtonsoft.Json" Version="13.0.3" />
+    <PackageReference Include="Newtonsoft.Json" Version="13.0.4" />
     <PackageReference Include="OpenTelemetry" Version="1.14.0" />
     <PackageReference Include="OpenTelemetry.Exporter.Console" Version="1.14.0" />
     <PackageReference Include="OpenTelemetry.Exporter.OpenTelemetryProtocol" Version="1.14.0" />
diff --git a/Engine/Source/Programs/UnrealCloudDDC/Jupiter/Jupiter.csproj b/Engine/Source/Programs/UnrealCloudDDC/Jupiter/Jupiter.csproj
--- a/Engine/Source/Programs/UnrealCloudDDC/Jupiter/Jupiter.csproj
+++ b/Engine/Source/Programs/UnrealCloudDDC/Jupiter/Jupiter.csproj
@@ -77,7 +77,7 @@
     <PackageReference Include="Azure.Storage.Blobs" Version="12.15.1" />
     <PackageReference Include="KubernetesClient" Version="17.0.14" />
     <PackageReference Include="MongoDB.Driver" Version="2.23.1" />
-    <PackageReference Include="Newtonsoft.Json" Version="13.0.3" />
+    <PackageReference Include="Newtonsoft.Json" Version="13.0.4" />
     <PackageReference Include="Serilog.Enrichers.ClientInfo" Version="2.1.1" />
     <PackageReference Include="SerilogAnalyzer" Version="0.15.0" />
     <PackageReference Include="Crc32.NET" Version="1.2.0" />
