Cloud Storage - Added retry when fetching blob references to help recover when Scylla is slow to respond (primarilly a issue when GC runs at the same time as the compute expensive metrics gathering)

#Synced-CL 52616557

diff --git a/Engine/Source/Programs/UnrealCloudDDC/Jupiter/Implementation/Blob/ScyllaBlobIndex.cs b/Engine/Source/Programs/UnrealCloudDDC/Jupiter/Implementation/Blob/ScyllaBlobIndex.cs
--- a/Engine/Source/Programs/UnrealCloudDDC/Jupiter/Implementation/Blob/ScyllaBlobIndex.cs
+++ b/Engine/Source/Programs/UnrealCloudDDC/Jupiter/Implementation/Blob/ScyllaBlobIndex.cs
@@ -371,16 +371,40 @@ public class ScyllaBlobIndex : IBlobIndex
 
 		bool noReferencesFound = true;
 
-		IEnumerable<ScyllaBlobIncomingReference> references;
-		try
+		IEnumerable<ScyllaBlobIncomingReference>? references = null;
+		const int MaxRetryAttempts = 3;
+		Exception? timeoutException = null;
+		for (int i = 0; i < MaxRetryAttempts; i++)
 		{
-			references = await _mapper.FetchAsync<ScyllaBlobIncomingReference>("WHERE namespace = ? AND blob_id = ? " + cqlOptions, ns.ToString(), id.HashData);
+			try
+			{
+				timeoutException = null;
+
+				references = await _mapper.FetchAsync<ScyllaBlobIncomingReference>("WHERE namespace = ? AND blob_id = ? " + cqlOptions, ns.ToString(), id.HashData);
+				break;
+			}
+			catch (ReadTimeoutException e)
+			{
+				// wait 2 seconds and try again as the Db is under heavy load right now
+				await Task.Delay(TimeSpan.FromSeconds(2), cancellationToken);
+				timeoutException = e;
+			}
+			catch (QueryExecutionException e)
+			{
+				throw new RemoteServiceException("Error communicating with Scylla", e);
+			}
 		}
-		catch (QueryExecutionException e)
+
+		if (timeoutException != null)
 		{
-			throw new RemoteServiceException("Error communicating with Scylla", e);
+			throw new RemoteServiceException("Error communicating with Scylla", timeoutException);
 		}
-		
+
+		if (references == null)
+		{
+			throw new NotImplementedException("References not set, unhandled exception?");
+		}
+
 		foreach (ScyllaBlobIncomingReference incomingReference in references)
 		{
 			noReferencesFound = false;
